UI Tweaks for Zen Stats

diff --git a/Engine/Source/Editor/DerivedDataEditor/Private/SZenCacheStatistics.cpp b/Engine/Source/Editor/DerivedDataEditor/Private/SZenCacheStatistics.cpp
--- a/Engine/Source/Editor/DerivedDataEditor/Private/SZenCacheStatistics.cpp
+++ b/Engine/Source/Editor/DerivedDataEditor/Private/SZenCacheStatistics.cpp
@@ -196,7 +196,7 @@ TSharedRef<SWidget> SZenCacheStatisticsDialog::GetGridPanel()
 		[
 			SNew(STextBlock)
 			.Margin(FMargin(ColumnMargin, RowMargin))
-			.Text( EndpointStats.Name.Contains("Jupiter")? FText::FromString(TEXT("Horde")) : FText::FromString(TEXT("Zen")) )
+			.Text_Lambda([EndpointStats] { return FText::FromString(EndpointStats.Name); })
 		];
 
 		Panel->AddSlot(1, Row)
@@ -231,7 +231,7 @@ TSharedRef<SWidget> SZenCacheStatisticsDialog::GetGridPanel()
 		[
 			SNew(STextBlock)
 			.Margin(FMargin(ColumnMargin, RowMargin))
-			.Text_Lambda([EndpointStats] { return FText::FromString(EndpointStats.Name); })
+			.Text_Lambda([EndpointStats] { return FText::FromString(EndpointStats.Url); })
 		];
 
 		Row++;
diff --git a/Engine/Source/Programs/ZenDashboard/Private/SZenCacheStatistics.cpp b/Engine/Source/Programs/ZenDashboard/Private/SZenCacheStatistics.cpp
--- a/Engine/Source/Programs/ZenDashboard/Private/SZenCacheStatistics.cpp
+++ b/Engine/Source/Programs/ZenDashboard/Private/SZenCacheStatistics.cpp
@@ -171,7 +171,7 @@ TSharedRef<SWidget> SZenCacheStatisticsDialog::GetGridPanel()
 		[
 			SNew(STextBlock)
 			.Margin(FMargin(ColumnMargin, RowMargin))
-			.Text( EndpointStats.Name.Contains("Jupiter")? FText::FromString(TEXT("Horde")) : FText::FromString(TEXT("Zen")) )
+			.Text_Lambda([EndpointStats] { return FText::FromString(EndpointStats.Name); })
 		];
 
 		Panel->AddSlot(1, Row)
@@ -206,7 +206,7 @@ TSharedRef<SWidget> SZenCacheStatisticsDialog::GetGridPanel()
 		[
 			SNew(STextBlock)
 			.Margin(FMargin(ColumnMargin, RowMargin))
-			.Text_Lambda([EndpointStats] { return FText::FromString(EndpointStats.Name); })
+			.Text_Lambda([EndpointStats] { return FText::FromString(EndpointStats.Url); })
 		];
 
 		SumTotalGetMB += EndpointStats.DownloadedMB;
