Updated Zen Dashbaord and Zen Stats for latest Zen Server API changes.
Removed ZenCacheStatistics from the DDC Editor Module as it is no longer required

diff --git a/Engine/Source/Developer/Zen/Internal/ZenServerInterface.cpp b/Engine/Source/Developer/Zen/Internal/ZenServerInterface.cpp
--- a/Engine/Source/Developer/Zen/Internal/ZenServerInterface.cpp
+++ b/Engine/Source/Developer/Zen/Internal/ZenServerInterface.cpp
@@ -821,6 +821,11 @@ FZenServiceInstance::GetStats( FZenStats& stats ) const
 		CacheStats.UpstreamHits = CacheObjectView["upstream_hits"].AsInt64();
 		CacheStats.UpstreamRatio = CacheObjectView["upstream_ratio"].AsDouble();
 
+		FCbObjectView CacheSizeObjectView = CacheObjectView["size"].AsObjectView();
+		FZenCacheSizeStats& CacheSizeStats = CacheStats.Size;
+		CacheSizeStats.Disk = CacheSizeObjectView["disk"].AsDouble();
+		CacheSizeStats.Memory = CacheSizeObjectView["memory"].AsDouble();
+
 		FCbObjectView UpstreamObjectView = RootObjectView["upstream"].AsObjectView();
 		FZenUpstreamStats& UpstreamStats = stats.UpstreamStats;
 
@@ -861,6 +866,16 @@ FZenServiceInstance::GetStats( FZenStats& stats ) const
 			UpstreamStats.EndPointStats.Push(EndPointStats);
 		}
 
+		FCbObjectView CASObjectView = RootObjectView["cas"].AsObjectView();	
+		FCbObjectView CASSizeObjectView = CASObjectView["size"].AsObjectView();
+
+		FZenCASSizeStats& CASSizeStats = stats.CASStats.Size;
+
+		CASSizeStats.Tiny = CASSizeObjectView["tiny"].AsInt64();
+		CASSizeStats.Small = CASSizeObjectView["small"].AsInt64();
+		CASSizeStats.Large = CASSizeObjectView["large"].AsInt64();
+		CASSizeStats.Total = CASSizeObjectView["total"].AsInt64();
+
 		return true;
 	}
 
diff --git a/Engine/Source/Developer/Zen/Public/ZenStatistics.h b/Engine/Source/Developer/Zen/Public/ZenStatistics.h
--- a/Engine/Source/Developer/Zen/Public/ZenStatistics.h
+++ b/Engine/Source/Developer/Zen/Public/ZenStatistics.h
@@ -6,13 +6,20 @@
 
 namespace UE::Zen {
 
+	struct FZenCacheSizeStats
+	{
+		double Disk = 0;
+		double Memory = 0;
+	};
+
 	struct FZenCacheStats
 	{
+		FZenCacheSizeStats Size;
 		int64 Hits = 0;
 		int64 Misses = 0;
 		double HitRatio = 0.0;
 		int64 UpstreamHits = 0;
-		double UpstreamRatio = 0.0;
+		double UpstreamRatio = 0.0;	
 	};
 
 	struct FZenRequestStats
@@ -46,12 +53,26 @@ namespace UE::Zen {
 		TArray<FZenEndPointStats> EndPointStats;
 	};
 
+	struct FZenCASSizeStats
+	{
+		int64 Tiny = 0;
+		int64 Small = 0;
+		int64 Large = 0;
+		int64 Total = 0;
+	};
+
+	struct FZenCASStats
+	{
+		FZenCASSizeStats Size;
+	};
+
 	struct FZenStats
 	{
 		FZenRequestStats RequestStats;
 		FZenCacheStats CacheStats;
 		FZenUpstreamStats UpstreamStats;
 		FZenRequestStats UpstreamRequestStats;
+		FZenCASStats CASStats;
 
 	};
 
diff --git a/Engine/Source/Editor/DerivedDataEditor/Private/DerivedDataEditorModule.cpp b/Engine/Source/Editor/DerivedDataEditor/Private/DerivedDataEditorModule.cpp
--- a/Engine/Source/Editor/DerivedDataEditor/Private/DerivedDataEditorModule.cpp
+++ b/Engine/Source/Editor/DerivedDataEditor/Private/DerivedDataEditorModule.cpp
@@ -5,7 +5,6 @@
 #include "SDerivedDataStatusBar.h"
 #include "SDerivedDataDialogs.h"
 #include "SDerivedDataCacheSettings.h"
-#include "SZenCacheStatistics.h"
 #include "SVirtualAssetsStatistics.h"
 #include "ZenServerInterface.h"
 #include "DerivedDataCacheNotifications.h"
diff --git a/Engine/Source/Editor/DerivedDataEditor/Private/SZenCacheStatistics.cpp b/Engine/Source/Editor/DerivedDataEditor/Private/SZenCacheStatistics.cpp
deleted file mode 100644
diff --git a/Engine/Source/Editor/DerivedDataEditor/Private/SZenCacheStatistics.h b/Engine/Source/Editor/DerivedDataEditor/Private/SZenCacheStatistics.h
deleted file mode 100644
diff --git a/Engine/Source/Programs/ZenDashboard/Private/SZenCacheStatistics.cpp b/Engine/Source/Programs/ZenDashboard/Private/SZenCacheStatistics.cpp
--- a/Engine/Source/Programs/ZenDashboard/Private/SZenCacheStatistics.cpp
+++ b/Engine/Source/Programs/ZenDashboard/Private/SZenCacheStatistics.cpp
@@ -30,7 +30,7 @@
 
 		+ SVerticalBox::Slot()
 		.AutoHeight()
-		.Padding(0, 20, 0, 0)
+		.Padding(0, 0, 0, 0)
 		.Expose(GridSlot)
 		[
 			GetGridPanel()
@@ -72,95 +72,112 @@
 	const FSlateColor TitleColor = FStyleColors::AccentWhite;
 	const FSlateFontInfo TitleFont = FCoreStyle::GetDefaultFontStyle("Bold", 10);
 
+	const double CacheDiskUsageMB = FUnitConversion::Convert(ZenStats.CacheStats.Size.Disk, EUnit::Bytes, EUnit::Megabytes);
+	const double CacheMemoryUsageMB = FUnitConversion::Convert(ZenStats.CacheStats.Size.Memory, EUnit::Bytes, EUnit::Megabytes);
+	
 	Panel->AddSlot(0, Row)
 	[
 		SNew(STextBlock)
-		.Margin(FMargin(ColumnMargin, RowMargin, 0.0f, TitleMargin))
+		.Margin(FMargin(ColumnMargin, RowMargin))
 		.ColorAndOpacity(TitleColor)
 		.Font(TitleFont)
-		.Text(LOCTEXT("Cache", "Cache"))
-		.Text(LOCTEXT("CacheType", "Cache Type"))
+		.Text(LOCTEXT("LocalServer", "Local Cache"))
+	];
+
+	Row++;
+
+	Panel->AddSlot(0, Row)
+	[
+		SNew(STextBlock)
+		.Margin(FMargin(ColumnMargin, RowMargin))
+		.Text(LOCTEXT("CacheSize", "Size"))
 	];
 
 	Panel->AddSlot(1, Row)
 	[
 		SNew(STextBlock)
-		.Margin(FMargin(ColumnMargin, RowMargin, 0.0f, TitleMargin))
+		.Margin(FMargin(ColumnMargin, RowMargin))
+		.Text_Lambda([CacheDiskUsageMB] { return FText::FromString(SingleDecimalFormat(CacheDiskUsageMB) + TEXT(" MB")); })
+	];
+
+	Row++;
+
+	Panel->AddSlot(0, Row)
+	[
+		SNew(STextBlock)
+		.Margin(FMargin(ColumnMargin, RowMargin))
+		.Text(LOCTEXT("CacheMemory", "Memory"))
+	];
+
+	Panel->AddSlot(1, Row)
+	[
+		SNew(STextBlock)
+		.Margin(FMargin(ColumnMargin, RowMargin))
+		.Text_Lambda([CacheMemoryUsageMB] { return FText::FromString(SingleDecimalFormat(CacheMemoryUsageMB) + TEXT(" MB")); })
+	];
+
+	Row++;
+
+	Panel->AddSlot(0, Row)
+	[
+		SNew(STextBlock)
+		.Margin(FMargin(ColumnMargin, RowMargin))
+		.Text(LOCTEXT("CacheHitRate", "Hit Rate"))
+	];
+
+	Panel->AddSlot(1, Row)
+	[
+		SNew(STextBlock)
+		.Margin(FMargin(ColumnMargin, RowMargin))
+		.Text_Lambda([ZenStats] { return FText::FromString(SingleDecimalFormat(ZenStats.CacheStats.HitRatio * 100.0) + TEXT(" %")); })
+	];
+
+	Row++;
+
+	Panel->AddSlot(0, Row)
+	[
+		SNew(STextBlock)
+		.Margin(FMargin(ColumnMargin, TitleMargin, ColumnMargin, 0.0f))
 		.ColorAndOpacity(TitleColor)
 		.Font(TitleFont)
-		.Justification(ETextJustify::Left)
-		.Text(LOCTEXT("Location", "Location"))
+		.Text(LOCTEXT("UpstreamServer", "Upstream Server"))
+	];
+
+	Panel->AddSlot(1, Row)
+	[
+		SNew(STextBlock)
+		.Margin(FMargin(ColumnMargin, TitleMargin, ColumnMargin, 0.0f))
+		.ColorAndOpacity(TitleColor)
+		.Font(TitleFont)
+		.Text(LOCTEXT("HitRate", "Hit Rate"))
 	];
 
 	Panel->AddSlot(2, Row)
 	[
 		SNew(STextBlock)
-		.Margin(FMargin(ColumnMargin, RowMargin, 0.0f, TitleMargin))
+		.Margin(FMargin(ColumnMargin, TitleMargin, ColumnMargin, 0.0f))
 		.ColorAndOpacity(TitleColor)
 		.Font(TitleFont)
-		.Justification(ETextJustify::Left)
-		.Text(LOCTEXT("HitPercentage", "Hit%"))
+		.Text(LOCTEXT("Downloaded", "Downloaded"))
 	];
 
 	Panel->AddSlot(3, Row)
 	[
 		SNew(STextBlock)
-		.Margin(FMargin(ColumnMargin, RowMargin, 0.0f, TitleMargin))
+		.Margin(FMargin(ColumnMargin, TitleMargin, ColumnMargin, 0.0f))
 		.ColorAndOpacity(TitleColor)
 		.Font(TitleFont)
-		.Justification(ETextJustify::Left)
-		.Text(LOCTEXT("Read", "Read"))
+		.Text(LOCTEXT("Uploaded", "Uploaded"))
 	];
 
 	Panel->AddSlot(4, Row)
-		[
+	[
 		SNew(STextBlock)
-		.Margin(FMargin(ColumnMargin, RowMargin, 0.0f, TitleMargin))
+		.Margin(FMargin(ColumnMargin, TitleMargin, ColumnMargin, 0.0f))
 		.ColorAndOpacity(TitleColor)
 		.Font(TitleFont)
-		.Justification(ETextJustify::Left)
-		.Text(LOCTEXT("Write", "Write"))
-		];
-
-	Panel->AddSlot(5, Row)
-		[
-			SNew(STextBlock)
-			.Margin(FMargin(ColumnMargin, RowMargin, 0.0f, TitleMargin))
-			.ColorAndOpacity(TitleColor)
-			.Font(TitleFont)
-			.Justification(ETextJustify::Left)
-			.Text(LOCTEXT("Details", "Details"))
-		];
-	Row++;
-
-	Panel->AddSlot(0, Row)
-	[
-		SNew(STextBlock)
-		.Margin(FMargin(ColumnMargin, RowMargin))
-		.Text(LOCTEXT("ZenServer", "Zen"))
+		.Text(LOCTEXT("URL", "URL"))
 	];
-
-	Panel->AddSlot(1, Row)
-	[
-		SNew(STextBlock)
-		.Margin(FMargin(ColumnMargin, RowMargin))
-		.Text(LOCTEXT("LocalServer", "Local"))
-	];
-
-	Panel->AddSlot(2, Row)
-	[
-		SNew(STextBlock)
-		.Margin(FMargin(ColumnMargin, RowMargin))
-		.Text_Lambda([ZenStats] { return FText::FromString(SingleDecimalFormat(ZenStats.CacheStats.HitRatio * 100.0) + TEXT(" %")); })
-	];
-
-	Panel->AddSlot(5, Row)
-	[
-		SNew(STextBlock)
-		.Margin(FMargin(ColumnMargin, RowMargin))
-		.Text_Lambda([this] { return FText::FromString(FString::Printf(TEXT("%s:%d"), UE::Zen::GetDefaultServiceInstance().GetHostName(), UE::Zen::GetDefaultServiceInstance().GetPort() )); })
-	];
-
 	Row++;
 	
 	int32 EndpointIndex = 1;
@@ -178,31 +195,24 @@
 		[
 			SNew(STextBlock)
 			.Margin(FMargin(ColumnMargin, RowMargin))
-			.Text(LOCTEXT("RemoteServer", "Remote"))
-		];
-
-		Panel->AddSlot(2, Row)
-		[
-			SNew(STextBlock)
-			.Margin(FMargin(ColumnMargin, RowMargin))
 			.Text_Lambda([EndpointStats] { return FText::FromString(SingleDecimalFormat(EndpointStats.HitRatio * 100.0) + TEXT(" %")); })
 		];
 
-		Panel->AddSlot(3, Row)
+		Panel->AddSlot(2, Row)
 		[
 			SNew(STextBlock)
 			.Margin(FMargin(ColumnMargin, RowMargin))
 			.Text_Lambda([EndpointStats] { return FText::FromString(SingleDecimalFormat(EndpointStats.DownloadedMB) + TEXT(" MB")); })
 		];
 
-		Panel->AddSlot(4, Row)
+		Panel->AddSlot(3, Row)
 		[
 			SNew(STextBlock)
 			.Margin(FMargin(ColumnMargin, RowMargin))
 			.Text_Lambda([EndpointStats] { return FText::FromString(SingleDecimalFormat(EndpointStats.UploadedMB) + TEXT(" MB")); })
 		];
 
-		Panel->AddSlot(5, Row)
+		Panel->AddSlot(4, Row)
 		[
 			SNew(STextBlock)
 			.Margin(FMargin(ColumnMargin, RowMargin))
@@ -215,37 +225,6 @@
 		Row++;
 	}
 
-	Panel->AddSlot(0, Row)
-	[
-		SNew(STextBlock)
-		.Text(FText::FromString(TEXT("Total")))
-		.Margin(FMargin(ColumnMargin, RowMargin))
-		.ColorAndOpacity(TitleColor)
-		.Font(TitleFont)
-		.Justification(ETextJustify::Left)	
-	];
-
-	Panel->AddSlot(3, Row)
-	.HAlign(HAlign_Right)
-	[
-		SNew(STextBlock)
-		.Margin(FMargin(ColumnMargin, RowMargin))
-		.Justification(ETextJustify::Left)
-		.ColorAndOpacity(TitleColor)
-		.Font(TitleFont)
-		.Text(FText::FromString(SingleDecimalFormat(SumTotalGetMB) + TEXT(" MB")))
-		];
-
-	Panel->AddSlot(4, Row)
-	.HAlign(HAlign_Right)
-	[
-		SNew(STextBlock)
-		.Margin(FMargin(ColumnMargin, RowMargin))
-		.Justification(ETextJustify::Left)
-		.ColorAndOpacity(TitleColor)
-		.Font(TitleFont)
-		.Text(FText::FromString(SingleDecimalFormat(SumTotalPutMB) + TEXT(" MB")))
-	];
 #endif
 
 	return Panel;
