Changed the Cache size to report in GB rather than MB

diff --git a/Engine/Source/Programs/ZenDashboard/Private/SZenCacheStatistics.cpp b/Engine/Source/Programs/ZenDashboard/Private/SZenCacheStatistics.cpp
--- a/Engine/Source/Programs/ZenDashboard/Private/SZenCacheStatistics.cpp
+++ b/Engine/Source/Programs/ZenDashboard/Private/SZenCacheStatistics.cpp
@@ -72,7 +72,7 @@ TSharedRef<SWidget> SZenCacheStatisticsDialog::GetGridPanel()
 	const FSlateColor TitleColor = FStyleColors::AccentWhite;
 	const FSlateFontInfo TitleFont = FCoreStyle::GetDefaultFontStyle("Bold", 10);
 
-	const double CacheDiskUsageMB = FUnitConversion::Convert(ZenStats.CacheStats.Size.Disk, EUnit::Bytes, EUnit::Megabytes);
+	const double CacheDiskUsagGB = FUnitConversion::Convert(ZenStats.CacheStats.Size.Disk, EUnit::Bytes, EUnit::Gigabytes);
 	const double CacheMemoryUsageMB = FUnitConversion::Convert(ZenStats.CacheStats.Size.Memory, EUnit::Bytes, EUnit::Megabytes);
 	
 	Panel->AddSlot(0, Row)
@@ -97,7 +97,7 @@ TSharedRef<SWidget> SZenCacheStatisticsDialog::GetGridPanel()
 	[
 		SNew(STextBlock)
 		.Margin(FMargin(ColumnMargin, RowMargin))
-		.Text_Lambda([CacheDiskUsageMB] { return FText::FromString(SingleDecimalFormat(CacheDiskUsageMB) + TEXT(" MB")); })
+		.Text_Lambda([CacheDiskUsagGB] { return FText::FromString(SingleDecimalFormat(CacheDiskUsagGB) + TEXT(" GB")); })
 	];
 
 	Row++;
