Minor UI tweak to Zen Dashboard stats to convert between MB/GB automatically

diff --git a/Engine/Source/Programs/ZenDashboard/Private/SZenCacheStatistics.cpp b/Engine/Source/Programs/ZenDashboard/Private/SZenCacheStatistics.cpp
--- a/Engine/Source/Programs/ZenDashboard/Private/SZenCacheStatistics.cpp
+++ b/Engine/Source/Programs/ZenDashboard/Private/SZenCacheStatistics.cpp
@@ -72,7 +72,7 @@ TSharedRef<SWidget> SZenCacheStatisticsDialog::GetGridPanel()
 	const FSlateColor TitleColor = FStyleColors::AccentWhite;
 	const FSlateFontInfo TitleFont = FCoreStyle::GetDefaultFontStyle("Bold", 10);
 
-	const double CacheDiskUsagGB = FUnitConversion::Convert(ZenStats.CacheStats.Size.Disk, EUnit::Bytes, EUnit::Gigabytes);
+	const double CacheDiskUsagMB = FUnitConversion::Convert(ZenStats.CacheStats.Size.Disk, EUnit::Bytes, EUnit::Megabytes);
 	const double CacheMemoryUsageMB = FUnitConversion::Convert(ZenStats.CacheStats.Size.Memory, EUnit::Bytes, EUnit::Megabytes);
 	
 	Panel->AddSlot(0, Row)
@@ -90,14 +90,24 @@ TSharedRef<SWidget> SZenCacheStatisticsDialog::GetGridPanel()
 	[
 		SNew(STextBlock)
 		.Margin(FMargin(ColumnMargin, RowMargin))
-		.Text(LOCTEXT("CacheSize", "Size"))
+		.Text(LOCTEXT("DiskSpace", "Disk Space"))
 	];
 
 	Panel->AddSlot(1, Row)
 	[
 		SNew(STextBlock)
 		.Margin(FMargin(ColumnMargin, RowMargin))
-		.Text_Lambda([CacheDiskUsagGB] { return FText::FromString(SingleDecimalFormat(CacheDiskUsagGB) + TEXT(" GB")); })
+			.Text_Lambda([CacheDiskUsagMB]
+		{
+			if (CacheDiskUsagMB > 1024.0)
+			{
+				return FText::FromString(SingleDecimalFormat(FUnitConversion::Convert(CacheDiskUsagMB, EUnit::Megabytes, EUnit::Gigabytes)) + TEXT(" GB"));
+			}
+			else
+			{
+				return FText::FromString(SingleDecimalFormat(CacheDiskUsagMB) + TEXT(" MB"));
+			}
+		})
 	];
 
 	Row++;
@@ -106,14 +116,24 @@ TSharedRef<SWidget> SZenCacheStatisticsDialog::GetGridPanel()
 	[
 		SNew(STextBlock)
 		.Margin(FMargin(ColumnMargin, RowMargin))
-		.Text(LOCTEXT("CacheMemory", "Memory"))
+		.Text(LOCTEXT("MemoryUsage", "Memory Usage"))
 	];
 
 	Panel->AddSlot(1, Row)
 	[
 		SNew(STextBlock)
 		.Margin(FMargin(ColumnMargin, RowMargin))
-		.Text_Lambda([CacheMemoryUsageMB] { return FText::FromString(SingleDecimalFormat(CacheMemoryUsageMB) + TEXT(" MB")); })
+		.Text_Lambda([CacheMemoryUsageMB] 
+			{
+				if (CacheMemoryUsageMB > 1024.0)
+				{
+					return FText::FromString(SingleDecimalFormat(FUnitConversion::Convert(CacheMemoryUsageMB, EUnit::Megabytes, EUnit::Gigabytes)) + TEXT(" GB"));
+				}
+				else
+				{
+					return FText::FromString(SingleDecimalFormat(CacheMemoryUsageMB) + TEXT(" MB"));
+				}
+			})
 	];
 
 	Row++;
@@ -202,14 +222,35 @@ TSharedRef<SWidget> SZenCacheStatisticsDialog::GetGridPanel()
 		[
 			SNew(STextBlock)
 			.Margin(FMargin(ColumnMargin, RowMargin))
-			.Text_Lambda([EndpointStats] { return FText::FromString(SingleDecimalFormat(EndpointStats.DownloadedMB) + TEXT(" MB")); })
+			.Text_Lambda([EndpointStats] 
+				{ 
+					if (EndpointStats.DownloadedMB> 1024.0)
+					{
+						return FText::FromString(SingleDecimalFormat(FUnitConversion::Convert(EndpointStats.DownloadedMB, EUnit::Megabytes, EUnit::Gigabytes) ) + TEXT(" GB")); 
+					}
+					else
+					{
+						return FText::FromString(SingleDecimalFormat(EndpointStats.DownloadedMB) + TEXT(" MB"));
+					}
+				})
+											
 		];
 
 		Panel->AddSlot(3, Row)
 		[
 			SNew(STextBlock)
 			.Margin(FMargin(ColumnMargin, RowMargin))
-			.Text_Lambda([EndpointStats] { return FText::FromString(SingleDecimalFormat(EndpointStats.UploadedMB) + TEXT(" MB")); })
+			.Text_Lambda([EndpointStats]
+				{ 
+					if (EndpointStats.UploadedMB > 1024.0)
+					{
+						return FText::FromString(SingleDecimalFormat(FUnitConversion::Convert(EndpointStats.UploadedMB, EUnit::Megabytes, EUnit::Gigabytes) ) + TEXT(" GB")); 
+					}
+					else
+					{
+						return FText::FromString(SingleDecimalFormat(EndpointStats.UploadedMB) + TEXT(" MB"));
+					}
+				})
 		];
 
 		Panel->AddSlot(4, Row)
