Tweak to Zen Cache Statistics display

diff --git a/Engine/Source/Programs/ZenDashboard/Private/SZenCacheStatistics.cpp b/Engine/Source/Programs/ZenDashboard/Private/SZenCacheStatistics.cpp
--- a/Engine/Source/Programs/ZenDashboard/Private/SZenCacheStatistics.cpp
+++ b/Engine/Source/Programs/ZenDashboard/Private/SZenCacheStatistics.cpp
@@ -72,16 +72,26 @@ TSharedRef<SWidget> SZenCacheStatisticsDialog::GetGridPanel()
 	const FSlateColor TitleColor = FStyleColors::AccentWhite;
 	const FSlateFontInfo TitleFont = FCoreStyle::GetDefaultFontStyle("Bold", 10);
 
-	const double CacheDiskUsagMB = FUnitConversion::Convert(ZenStats.CacheStats.Size.Disk, EUnit::Bytes, EUnit::Megabytes);
+	const double CASDiskUsageMB = FUnitConversion::Convert(ZenStats.CASStats.Size.Total, EUnit::Bytes, EUnit::Megabytes);
+	const double CacheDiskUsageMB = FUnitConversion::Convert(ZenStats.CacheStats.Size.Disk, EUnit::Bytes, EUnit::Megabytes);
 	const double CacheMemoryUsageMB = FUnitConversion::Convert(ZenStats.CacheStats.Size.Memory, EUnit::Bytes, EUnit::Megabytes);
-	
+
 	Panel->AddSlot(0, Row)
 	[
 		SNew(STextBlock)
 		.Margin(FMargin(ColumnMargin, RowMargin))
 		.ColorAndOpacity(TitleColor)
 		.Font(TitleFont)
-		.Text(LOCTEXT("LocalServer", "Local Cache"))
+		.Text(LOCTEXT("Cache", "Local Cache"))
+	];
+
+	Panel->AddSlot(2, Row)
+	[
+		SNew(STextBlock)
+		.Margin(FMargin(ColumnMargin, RowMargin))
+		.ColorAndOpacity(TitleColor)
+		.Font(TitleFont)
+		.Text(LOCTEXT("CAS", "Local Content Store"))
 	];
 
 	Row++;
@@ -97,15 +107,39 @@ TSharedRef<SWidget> SZenCacheStatisticsDialog::GetGridPanel()
 	[
 		SNew(STextBlock)
 		.Margin(FMargin(ColumnMargin, RowMargin))
-			.Text_Lambda([CacheDiskUsagMB]
+		.Text_Lambda([CacheDiskUsageMB]
+		{
+			if (CacheDiskUsageMB > 1024.0)
+			{
+				return FText::FromString(SingleDecimalFormat(FUnitConversion::Convert(CacheDiskUsageMB, EUnit::Megabytes, EUnit::Gigabytes)) + TEXT(" GB"));
+			}
+			else
+			{
+				return FText::FromString(SingleDecimalFormat(CacheDiskUsageMB) + TEXT(" MB"));
+			}
+		})
+	];
+
+	Panel->AddSlot(2, Row)
+	[
+		SNew(STextBlock)
+		.Margin(FMargin(ColumnMargin, RowMargin))
+		.Text(LOCTEXT("DiskSpace", "Disk Space"))
+	];
+
+	Panel->AddSlot(3, Row)
+	[
+		SNew(STextBlock)
+		.Margin(FMargin(ColumnMargin, RowMargin))
+		.Text_Lambda([CASDiskUsageMB]
 		{
-			if (CacheDiskUsagMB > 1024.0)
+			if (CASDiskUsageMB > 1024.0)
 			{
-				return FText::FromString(SingleDecimalFormat(FUnitConversion::Convert(CacheDiskUsagMB, EUnit::Megabytes, EUnit::Gigabytes)) + TEXT(" GB"));
+				return FText::FromString(SingleDecimalFormat(FUnitConversion::Convert(CASDiskUsageMB, EUnit::Megabytes, EUnit::Gigabytes)) + TEXT(" GB"));
 			}
 			else
 			{
-				return FText::FromString(SingleDecimalFormat(CacheDiskUsagMB) + TEXT(" MB"));
+				return FText::FromString(SingleDecimalFormat(CASDiskUsageMB) + TEXT(" MB"));
 			}
 		})
 	];
