Zen dashboard feature additions for:
- Service status & controls
- Version information
- GC status & control

Depends on https://p4-swarm.epicgames.net/reviews/23337377

diff --git a/Engine/Programs/ZenDashboard/Config/DefaultEngine.ini b/Engine/Programs/ZenDashboard/Config/DefaultEngine.ini
new file mode 100644
--- /dev/null
+++ b/Engine/Programs/ZenDashboard/Config/DefaultEngine.ini
@@ -0,0 +1,2 @@
+[TaskGraph]
+EnableReserveWorkers=false
diff --git a/Engine/Source/Programs/ZenDashboard/Private/SMessageDialog.cpp b/Engine/Source/Programs/ZenDashboard/Private/SMessageDialog.cpp
new file mode 100644
--- /dev/null
+++ b/Engine/Source/Programs/ZenDashboard/Private/SMessageDialog.cpp
@@ -0,0 +1,426 @@
+// Copyright Epic Games, Inc. All Rights Reserved.
+
+
+#include "SMessageDialog.h"
+#include "Fonts/SlateFontInfo.h"
+#include "Framework/Application/SlateApplication.h"
+#include "HAL/PlatformApplicationMisc.h"
+#include "Input/Reply.h"
+#include "InputCoreTypes.h"
+#include "Misc/Attribute.h"
+#include "SlateOptMacros.h"
+#include "Styling/CoreStyle.h"
+#include "Widgets/DeclarativeSyntaxSupport.h"
+#include "Widgets/Input/SButton.h"
+#include "Widgets/Layout/SBorder.h"
+#include "Widgets/Layout/SScrollBox.h"
+#include "Widgets/Layout/SUniformGridPanel.h"
+#include "Widgets/SBoxPanel.h"
+#include "Widgets/SCompoundWidget.h"
+#include "Widgets/SWindow.h"
+#include "Widgets/Text/STextBlock.h"
+
+#define LOCTEXT_NAMESPACE "Dialogs"
+
+class SChoiceDialog : public SCompoundWidget
+{
+public:
+	SLATE_BEGIN_ARGS( SChoiceDialog )	{}
+		SLATE_ATTRIBUTE(TSharedPtr<SWindow>, ParentWindow)
+		SLATE_ATTRIBUTE(FText, Message)	
+		SLATE_ATTRIBUTE(float, WrapMessageAt)
+		SLATE_ATTRIBUTE(EAppMsgType::Type, MessageType)
+	SLATE_END_ARGS()
+
+	BEGIN_SLATE_FUNCTION_BUILD_OPTIMIZATION
+	void Construct( const FArguments& InArgs )
+	{
+		ParentWindow = InArgs._ParentWindow.Get();
+		ParentWindow->SetWidgetToFocusOnActivate(SharedThis(this));
+		Response = EAppReturnType::Cancel;
+		MessageType = InArgs._MessageType.Get();
+
+		FSlateFontInfo MessageFont( FAppStyle::GetFontStyle("StandardDialog.LargeFont"));
+		MyMessage = InArgs._Message;
+
+		TSharedPtr<SUniformGridPanel> ButtonBox;
+
+		this->ChildSlot
+		[	
+			SNew(SBorder)
+			.BorderImage(FAppStyle::GetBrush("ToolPanel.GroupBorder"))
+			.Padding(16.f)
+			[
+				SNew(SVerticalBox)
+
+				+ SVerticalBox::Slot()
+				.HAlign(HAlign_Fill)
+				.VAlign(VAlign_Fill)
+				.FillHeight(1.0f)
+				.MaxHeight(550)
+				[
+					SNew(SHorizontalBox)
+					+SHorizontalBox::Slot()
+					.AutoWidth()
+					.VAlign(VAlign_Top)
+					.HAlign(HAlign_Left)
+					[
+						SNew(SImage)
+						.DesiredSizeOverride(FVector2D(24.f, 24.f))
+						.Image(FAppStyle::Get().GetBrush("Icons.WarningWithColor.Large"))
+					]
+
+					+SHorizontalBox::Slot()
+					.Padding(16.f, 0.f, 0.f, 0.f)
+					[
+						SNew(SScrollBox)
+						+ SScrollBox::Slot()
+						.HAlign(HAlign_Left)
+						.VAlign(VAlign_Center)
+						[
+							SNew(STextBlock)
+							.Text(MyMessage)
+							.Font(MessageFont)
+							.WrapTextAt(InArgs._WrapMessageAt)
+						]
+					]
+				]
+
+				+SVerticalBox::Slot()
+				.AutoHeight()
+				[
+					SNew(SHorizontalBox)
+
+					+ SHorizontalBox::Slot()
+					.AutoWidth()
+					.HAlign(HAlign_Left)
+					.VAlign(VAlign_Center)
+					[
+						SNew(SButton)
+						.ButtonStyle(FAppStyle::Get(), "SimpleButton")
+						.OnClicked(this, &SChoiceDialog::HandleCopyMessageButtonClicked)
+						.ToolTipText(NSLOCTEXT("SChoiceDialog", "CopyMessageTooltip", "Copy the text in this message to the clipboard (CTRL+C)"))
+						.ContentPadding(2.f)
+						.Content()
+						[
+							SNew(SImage)
+							.Image(FAppStyle::Get().GetBrush("Icons.Clipboard"))
+							.ColorAndOpacity(FSlateColor::UseForeground())
+						]
+					]
+
+					+ SHorizontalBox::Slot()
+						.FillWidth(1.0f)
+						.HAlign(HAlign_Left)
+						.VAlign(VAlign_Center)
+						.Padding(FMargin(16.f, 0.f, 0.f, 0.f))
+						[
+							SNew(SCheckBox)
+							.IsChecked(ECheckBoxState::Unchecked)
+							.OnCheckStateChanged(this, &SChoiceDialog::OnCheckboxClicked)
+							.Visibility(this, &SChoiceDialog::GetCheckboxVisibility)
+							.ToolTipText(NSLOCTEXT("SChoiceDialog", "ApplyToAllTooltip", "Make your choice of Yes or No apply to all remaining items in the current operation"))
+							[
+								SNew(STextBlock)
+								.WrapTextAt(615.0f)
+								.Text(NSLOCTEXT("SChoiceDialog", "ApplyToAllLabel", "Apply to All"))
+							]
+						]
+
+					+ SHorizontalBox::Slot()
+					.AutoWidth()
+					.HAlign(HAlign_Right)
+					.VAlign(VAlign_Center)
+					[
+						SAssignNew( ButtonBox, SUniformGridPanel )
+						.SlotPadding(FMargin(16.f, 0.f, 0.f, 0.f))
+						.MinDesiredSlotWidth(FAppStyle::Get().GetFloat("StandardDialog.MinDesiredSlotWidth"))
+						.MinDesiredSlotHeight(FAppStyle::Get().GetFloat("StandardDialog.MinDesiredSlotHeight"))
+					]
+				]
+			]
+		];
+
+		int32 SlotIndex = 0;
+
+#define ADD_SLOT(Button)\
+		ButtonBox->AddSlot(SlotIndex++,0)\
+		[\
+			SNew( SButton )\
+			.VAlign(VAlign_Center)\
+			.HAlign(HAlign_Center)\
+			.Text( EAppReturnTypeToText(EAppReturnType::Button) )\
+			.OnClicked( this, &SChoiceDialog::HandleButtonClicked, EAppReturnType::Button )\
+		];
+
+#define ADD_SLOT_PRIMARY(Button)\
+		ButtonBox->AddSlot(SlotIndex++,0)\
+		[\
+			SNew( SButton )\
+			.VAlign(VAlign_Center)\
+			.HAlign(HAlign_Center)\
+			.ButtonStyle(&FAppStyle::Get().GetWidgetStyle<FButtonStyle>("PrimaryButton"))\
+			.Text( EAppReturnTypeToText(EAppReturnType::Button) )\
+			.OnClicked( this, &SChoiceDialog::HandleButtonClicked, EAppReturnType::Button )\
+		];
+
+
+
+		switch ( MessageType )
+		{	
+		case EAppMsgType::Ok:
+			ADD_SLOT_PRIMARY(Ok)
+			break;
+		case EAppMsgType::YesNo:
+			ADD_SLOT_PRIMARY(Yes)
+			ADD_SLOT(No)
+			break;
+		case EAppMsgType::OkCancel:
+			ADD_SLOT_PRIMARY(Ok)
+			ADD_SLOT(Cancel)
+			break;
+		case EAppMsgType::YesNoCancel:
+			ADD_SLOT_PRIMARY(Yes)
+			ADD_SLOT(No)
+			ADD_SLOT(Cancel)
+			break;
+		case EAppMsgType::CancelRetryContinue:
+			ADD_SLOT_PRIMARY(Continue)
+			ADD_SLOT(Retry)
+			ADD_SLOT(Cancel)
+			break;
+		case EAppMsgType::YesNoYesAllNoAll:
+			ADD_SLOT_PRIMARY(Yes)
+			ADD_SLOT(No)
+			break;
+		case EAppMsgType::YesNoYesAllNoAllCancel:
+			ADD_SLOT_PRIMARY(Yes)
+			ADD_SLOT(No)
+			ADD_SLOT(Cancel)
+			break;
+		case EAppMsgType::YesNoYesAll:
+			ADD_SLOT_PRIMARY(Yes)
+			ADD_SLOT(No)
+			ADD_SLOT(YesAll)
+			break;
+		}
+
+#undef ADD_SLOT
+#undef ADD_SLOT_PRIMARY
+	}
+	END_SLATE_FUNCTION_BUILD_OPTIMIZATION
+
+
+	EAppReturnType::Type GetResponse()
+	{
+		return Response;
+	}
+
+	virtual	FReply OnKeyDown( const FGeometry& MyGeometry, const FKeyEvent& InKeyEvent )
+	{
+		if( InKeyEvent.GetKey() == EKeys::Escape )
+		{
+			return HandleButtonClicked(EAppReturnType::Cancel);
+		}
+
+		if (InKeyEvent.GetKey() == EKeys::C && InKeyEvent.IsControlDown())
+		{
+			CopyMessageToClipboard();
+
+			return FReply::Handled();
+		}
+
+		//if it was some other button, ignore it
+		return FReply::Unhandled();
+	}
+
+	/** Override the base method to allow for keyboard focus */
+	virtual bool SupportsKeyboardFocus() const
+	{
+		return true;
+	}
+
+	/** Converts an EAppReturnType into a localized FText */
+	static FText EAppReturnTypeToText(EAppReturnType::Type ReturnType)
+	{
+		switch(ReturnType)
+		{
+		case EAppReturnType::No:
+			return LOCTEXT("EAppReturnTypeNo", "No");
+		case EAppReturnType::Yes:
+			return LOCTEXT("EAppReturnTypeYes", "Yes");
+		case EAppReturnType::YesAll:
+			return LOCTEXT("EAppReturnTypeYesAll", "Yes All");
+		case EAppReturnType::NoAll:
+			return LOCTEXT("EAppReturnTypeNoAll", "No All");
+		case EAppReturnType::Cancel:
+			return LOCTEXT("EAppReturnTypeCancel", "Cancel");
+		case EAppReturnType::Ok:
+			return LOCTEXT("EAppReturnTypeOk", "OK");
+		case EAppReturnType::Retry:
+			return LOCTEXT("EAppReturnTypeRetry", "Retry");
+		case EAppReturnType::Continue:
+			return LOCTEXT("EAppReturnTypeContinue", "Continue");
+		default:
+			return LOCTEXT("MissingType", "MISSING RETURN TYPE");
+		}
+	}
+
+protected:
+
+	/**
+	 * Copies the message text to the clipboard.
+	 */
+	void CopyMessageToClipboard( )
+	{
+		FPlatformApplicationMisc::ClipboardCopy( *MyMessage.Get().ToString() );
+	}
+
+
+private:
+
+	// Handles clicking a message box button.
+	FReply HandleButtonClicked( EAppReturnType::Type InResponse )
+	{
+		Response = InResponse;
+		if ((MessageType == EAppMsgType::YesNoYesAllNoAll || MessageType == EAppMsgType::YesNoYesAllNoAllCancel)
+			&& bApplyToAllChecked)
+		{
+			if (Response == EAppReturnType::Yes)
+			{
+				Response = EAppReturnType::YesAll;
+			}
+			else if (Response == EAppReturnType::No)
+			{
+				Response = EAppReturnType::NoAll;
+			}
+		}
+
+		ResultCallback.ExecuteIfBound(ParentWindow.ToSharedRef(), Response);
+
+
+		ParentWindow->RequestDestroyWindow();
+
+		return FReply::Handled();
+	}
+
+	// Handles clicking the 'Copy Message' button.
+	FReply HandleCopyMessageButtonClicked()
+	{
+		CopyMessageToClipboard();
+		return FReply::Handled();
+	}
+		
+	// Used as a delegate for the OnClicked property of the Apply to All checkbox
+	void OnCheckboxClicked(ECheckBoxState InNewState)
+	{
+		bApplyToAllChecked = InNewState == ECheckBoxState::Checked;
+	}
+
+	// Used as a delegate for the Visibility property of the Apply to All checkbox
+	EVisibility GetCheckboxVisibility() const
+	{
+		if (MessageType == EAppMsgType::YesNoYesAllNoAll || MessageType == EAppMsgType::YesNoYesAllNoAllCancel)
+		{
+			return EVisibility::Visible;
+		}
+		return EVisibility::Hidden;
+	}
+
+public:
+	/** Callback delegate that is triggered, when the dialog is run in non-modal mode */
+	FOnMsgDlgResult ResultCallback;
+
+private:
+
+	EAppReturnType::Type Response;
+	TSharedPtr<SWindow> ParentWindow;
+	TAttribute<FText> MyMessage;
+	bool bApplyToAllChecked;
+	EAppMsgType::Type MessageType;
+};
+
+
+void CreateMsgDlgWindow(TSharedPtr<SWindow>& OutWindow, TSharedPtr<SChoiceDialog>& OutDialog, EAppMsgType::Type InMessageType,
+						const FText& InMessage, const FText& InTitle, FOnMsgDlgResult ResultCallback=NULL)
+{
+	OutWindow = SNew(SWindow)
+		.Title(InTitle)
+		.SizingRule(ESizingRule::Autosized)
+		.AutoCenter(EAutoCenter::PreferredWorkArea)
+		.SupportsMinimize(false).SupportsMaximize(false);
+
+	OutDialog = SNew(SChoiceDialog)
+		.ParentWindow(OutWindow)
+		.Message(InMessage)
+		.WrapMessageAt(512.0f)
+		.MessageType(InMessageType);
+
+	OutDialog->ResultCallback = ResultCallback;
+
+	OutWindow->SetContent(OutDialog.ToSharedRef());
+}
+
+EAppReturnType::Type OpenModalMessageDialog_Internal(EAppMsgType::Type InMessageType, EAppReturnType::Type InDefaultValue, const FText& InMessage, const FText& InTitle, const TSharedPtr<const SWidget>& ModalParent)
+{
+	EAppReturnType::Type Response = InDefaultValue;
+	TSharedPtr<SWindow> MsgWindow = NULL;
+	TSharedPtr<SChoiceDialog> MsgDialog = NULL;
+
+	CreateMsgDlgWindow(MsgWindow, MsgDialog, InMessageType, InMessage, InTitle);
+
+	FSlateApplication::Get().AddModalWindow(MsgWindow.ToSharedRef(), ModalParent);
+
+	Response = MsgDialog->GetResponse();
+
+	return Response;
+}
+
+EAppReturnType::Type OpenModalMessageDialog_Internal(EAppMsgType::Type InMessageType, const FText& InMessage, const FText& InTitle, const TSharedPtr<const SWidget>& ModalParent)
+{
+	EAppReturnType::Type DefaultValue = EAppReturnType::Yes;
+	switch (InMessageType)
+	{
+	case EAppMsgType::Ok:
+		DefaultValue = EAppReturnType::Ok;
+		break;
+	case EAppMsgType::YesNo:
+		DefaultValue = EAppReturnType::No;
+		break;
+	case EAppMsgType::OkCancel:
+		DefaultValue = EAppReturnType::Cancel;
+		break;
+	case EAppMsgType::YesNoCancel:
+		DefaultValue = EAppReturnType::Cancel;
+		break;
+	case EAppMsgType::CancelRetryContinue:
+		DefaultValue = EAppReturnType::Cancel;
+		break;
+	case EAppMsgType::YesNoYesAllNoAll:
+		DefaultValue = EAppReturnType::No;
+		break;
+	case EAppMsgType::YesNoYesAllNoAllCancel:
+		DefaultValue = EAppReturnType::Cancel;
+		break;
+	default:
+		DefaultValue = EAppReturnType::Yes;
+		break;
+	}
+
+	if (GIsRunningUnattendedScript && InMessageType != EAppMsgType::Ok)
+	{
+		//UE_LOG(LogDialogs, Error, TEXT("Message Dialog was triggered in unattended script mode without a default value. %d will be used."), (int32)DefaultValue);
+		if (FPlatformMisc::IsDebuggerPresent())
+		{
+			UE_DEBUG_BREAK();
+		}
+		else
+		{
+			FDebug::DumpStackTraceToLog(ELogVerbosity::Error);
+		}
+	}
+
+	return OpenModalMessageDialog_Internal(InMessageType, DefaultValue, InMessage, InTitle, ModalParent);
+}
+
+#undef LOCTEXT_NAMESPACE
diff --git a/Engine/Source/Programs/ZenDashboard/Private/SMessageDialog.h b/Engine/Source/Programs/ZenDashboard/Private/SMessageDialog.h
new file mode 100644
--- /dev/null
+++ b/Engine/Source/Programs/ZenDashboard/Private/SMessageDialog.h
@@ -0,0 +1,15 @@
+// Copyright Epic Games, Inc. All Rights Reserved.
+
+#pragma once
+
+#include "Delegates/DelegateCombinations.h"
+#include "GenericPlatform/GenericPlatformMisc.h"
+#include "Internationalization/Text.h"
+
+class SWindow;
+
+// Largely a copy of UnrealEd Dialogs.h
+
+DECLARE_DELEGATE_TwoParams(FOnMsgDlgResult, const TSharedRef<SWindow>&, EAppReturnType::Type);
+
+EAppReturnType::Type OpenModalMessageDialog_Internal(EAppMsgType::Type InMessageType, const FText& InMessage, const FText& InTitle, const TSharedPtr<const SWidget>& ModalParent);
diff --git a/Engine/Source/Programs/ZenDashboard/Private/SZenCacheStatistics.cpp b/Engine/Source/Programs/ZenDashboard/Private/SZenCacheStatistics.cpp
--- a/Engine/Source/Programs/ZenDashboard/Private/SZenCacheStatistics.cpp
+++ b/Engine/Source/Programs/ZenDashboard/Private/SZenCacheStatistics.cpp
@@ -11,19 +11,12 @@
 #include "Widgets/SBoxPanel.h"
 #include "Internationalization/FastDecimalFormat.h"
 
-#define LOCTEXT_NAMESPACE "ZenEditor"
+#define LOCTEXT_NAMESPACE "ZenDashboard"
 
-static FString SingleDecimalFormat(double Value)
+void SZenCacheStatistics::Construct(const FArguments& InArgs)
 {
-	const FNumberFormattingOptions NumberFormattingOptions = FNumberFormattingOptions()
-		.SetUseGrouping(true)
-		.SetMinimumFractionalDigits(1)
-		.SetMaximumFractionalDigits(1);
-	return FastDecimalFormat::NumberToString(Value, ExpressionParser::GetLocalizedNumberFormattingRules(), NumberFormattingOptions);
-};
+	ZenServiceInstance = InArgs._ZenServiceInstance;
 
-void SZenCacheStatisticsDialog::Construct(const FArguments& InArgs)
-{
 	this->ChildSlot
 	[
 		SNew(SVerticalBox)
@@ -37,10 +30,10 @@ void SZenCacheStatisticsDialog::Construct(const FArguments& InArgs)
 		]
 	];
 
-	RegisterActiveTimer(0.5f, FWidgetActiveTimerDelegate::CreateSP(this, &SZenCacheStatisticsDialog::UpdateGridPanels));
+	RegisterActiveTimer(0.5f, FWidgetActiveTimerDelegate::CreateSP(this, &SZenCacheStatistics::UpdateGridPanels));
 }
 
-EActiveTimerReturnType SZenCacheStatisticsDialog::UpdateGridPanels(double InCurrentTime, float InDeltaTime)
+EActiveTimerReturnType SZenCacheStatistics::UpdateGridPanels(double InCurrentTime, float InDeltaTime)
 {
 	(*GridSlot)
 	[
@@ -52,19 +45,23 @@ EActiveTimerReturnType SZenCacheStatisticsDialog::UpdateGridPanels(double InCurr
 	return EActiveTimerReturnType::Continue;
 }
 
-TSharedRef<SWidget> SZenCacheStatisticsDialog::GetGridPanel()
+TSharedRef<SWidget> SZenCacheStatistics::GetGridPanel()
 {
 	TSharedRef<SGridPanel> Panel = SNew(SGridPanel);
 
-#if UE_WITH_ZEN
-
 	UE::Zen::FZenStats ZenStats;
 
-	UE::Zen::GetDefaultServiceInstance().GetStats(ZenStats);
+	if (TSharedPtr<UE::Zen::FZenServiceInstance> ServiceInstance = ZenServiceInstance.Get())
+	{
+		ServiceInstance->GetStats(ZenStats);
+	}
+
+	const static FNumberFormattingOptions SingleDecimalFormatting = FNumberFormattingOptions()
+		.SetUseGrouping(true)
+		.SetMinimumFractionalDigits(1)
+		.SetMaximumFractionalDigits(1);
 
 	int32 Row = 0;
-	double SumTotalGetMB = 0.0;
-	double SumTotalPutMB = 0.0;
 
 	const float RowMargin = 0.0f;
 	const float TitleMargin = 10.0f;
@@ -72,9 +69,6 @@ TSharedRef<SWidget> SZenCacheStatisticsDialog::GetGridPanel()
 	const FSlateColor TitleColor = FStyleColors::AccentWhite;
 	const FSlateFontInfo TitleFont = FCoreStyle::GetDefaultFontStyle("Bold", 10);
 
-	const double CASDiskUsageMB = FUnitConversion::Convert(ZenStats.CASStats.Size.Total, EUnit::Bytes, EUnit::Megabytes);
-	const double CacheDiskUsageMB = FUnitConversion::Convert(ZenStats.CacheStats.Size.Disk, EUnit::Bytes, EUnit::Megabytes);
-	const double CacheMemoryUsageMB = FUnitConversion::Convert(ZenStats.CacheStats.Size.Memory, EUnit::Bytes, EUnit::Megabytes);
 
 	Panel->AddSlot(0, Row)
 	[
@@ -107,16 +101,9 @@ TSharedRef<SWidget> SZenCacheStatisticsDialog::GetGridPanel()
 	[
 		SNew(STextBlock)
 		.Margin(FMargin(ColumnMargin, RowMargin))
-		.Text_Lambda([CacheDiskUsageMB]
+		.Text_Lambda([DiskUsage = ZenStats.CacheStats.Size.Disk]
 		{
-			if (CacheDiskUsageMB > 1024.0)
-			{
-				return FText::FromString(SingleDecimalFormat(FUnitConversion::Convert(CacheDiskUsageMB, EUnit::Megabytes, EUnit::Gigabytes)) + TEXT(" GB"));
-			}
-			else
-			{
-				return FText::FromString(SingleDecimalFormat(CacheDiskUsageMB) + TEXT(" MB"));
-			}
+			return FText::AsMemory(DiskUsage, (DiskUsage > 1024) ? &SingleDecimalFormatting : nullptr, nullptr, EMemoryUnitStandard::IEC);
 		})
 	];
 
@@ -131,16 +118,9 @@ TSharedRef<SWidget> SZenCacheStatisticsDialog::GetGridPanel()
 	[
 		SNew(STextBlock)
 		.Margin(FMargin(ColumnMargin, RowMargin))
-		.Text_Lambda([CASDiskUsageMB]
+		.Text_Lambda([DiskUsage = ZenStats.CASStats.Size.Total]
 		{
-			if (CASDiskUsageMB > 1024.0)
-			{
-				return FText::FromString(SingleDecimalFormat(FUnitConversion::Convert(CASDiskUsageMB, EUnit::Megabytes, EUnit::Gigabytes)) + TEXT(" GB"));
-			}
-			else
-			{
-				return FText::FromString(SingleDecimalFormat(CASDiskUsageMB) + TEXT(" MB"));
-			}
+			return FText::AsMemory(DiskUsage, (DiskUsage > 1024) ? &SingleDecimalFormatting : nullptr, nullptr, EMemoryUnitStandard::IEC);
 		})
 	];
 
@@ -150,24 +130,18 @@ TSharedRef<SWidget> SZenCacheStatisticsDialog::GetGridPanel()
 	[
 		SNew(STextBlock)
 		.Margin(FMargin(ColumnMargin, RowMargin))
-		.Text(LOCTEXT("MemoryUsage", "Memory Usage"))
+		.Text(LOCTEXT("CacheHitRate", "Hit Rate"))
 	];
 
 	Panel->AddSlot(1, Row)
 	[
 		SNew(STextBlock)
 		.Margin(FMargin(ColumnMargin, RowMargin))
-		.Text_Lambda([CacheMemoryUsageMB] 
-			{
-				if (CacheMemoryUsageMB > 1024.0)
-				{
-					return FText::FromString(SingleDecimalFormat(FUnitConversion::Convert(CacheMemoryUsageMB, EUnit::Megabytes, EUnit::Gigabytes)) + TEXT(" GB"));
-				}
-				else
-				{
-					return FText::FromString(SingleDecimalFormat(CacheMemoryUsageMB) + TEXT(" MB"));
-				}
-			})
+		.Text_Lambda([CacheStats = ZenStats.CacheStats]
+		{
+			int64 CacheTotal = CacheStats.Hits + CacheStats.Misses;
+			return CacheTotal == 0 ? LOCTEXT("CacheNoHitRateValue", "-") : FText::AsPercent(CacheStats.HitRatio, &SingleDecimalFormatting);
+		})
 	];
 
 	Row++;
@@ -176,131 +150,137 @@ TSharedRef<SWidget> SZenCacheStatisticsDialog::GetGridPanel()
 	[
 		SNew(STextBlock)
 		.Margin(FMargin(ColumnMargin, RowMargin))
-		.Text(LOCTEXT("CacheHitRate", "Hit Rate"))
+		.Text(LOCTEXT("CacheHitQuantity", "Hits"))
 	];
 
 	Panel->AddSlot(1, Row)
 	[
 		SNew(STextBlock)
 		.Margin(FMargin(ColumnMargin, RowMargin))
-		.Text_Lambda([ZenStats] { return FText::FromString(SingleDecimalFormat(ZenStats.CacheStats.HitRatio * 100.0) + TEXT(" %")); })
+		.Text_Lambda([CacheHits = ZenStats.CacheStats.Hits]
+		{
+			return FText::AsNumber(CacheHits);
+		})
 	];
 
 	Row++;
 
 	Panel->AddSlot(0, Row)
-	[
-		SNew(STextBlock)
-		.Margin(FMargin(ColumnMargin, TitleMargin, ColumnMargin, 0.0f))
-		.ColorAndOpacity(TitleColor)
-		.Font(TitleFont)
-		.Text(LOCTEXT("UpstreamServer", "Upstream Server"))
-	];
+		[
+			SNew(STextBlock)
+			.Margin(FMargin(ColumnMargin, RowMargin))
+		.Text(LOCTEXT("CacheMissQuantity", "Misses"))
+		];
 
 	Panel->AddSlot(1, Row)
-	[
-		SNew(STextBlock)
-		.Margin(FMargin(ColumnMargin, TitleMargin, ColumnMargin, 0.0f))
-		.ColorAndOpacity(TitleColor)
-		.Font(TitleFont)
-		.Text(LOCTEXT("HitRate", "Hit Rate"))
-	];
-
-	Panel->AddSlot(2, Row)
-	[
-		SNew(STextBlock)
-		.Margin(FMargin(ColumnMargin, TitleMargin, ColumnMargin, 0.0f))
-		.ColorAndOpacity(TitleColor)
-		.Font(TitleFont)
-		.Text(LOCTEXT("Downloaded", "Downloaded"))
-	];
-
-	Panel->AddSlot(3, Row)
-	[
-		SNew(STextBlock)
-		.Margin(FMargin(ColumnMargin, TitleMargin, ColumnMargin, 0.0f))
-		.ColorAndOpacity(TitleColor)
-		.Font(TitleFont)
-		.Text(LOCTEXT("Uploaded", "Uploaded"))
-	];
+		[
+			SNew(STextBlock)
+			.Margin(FMargin(ColumnMargin, RowMargin))
+		.Text_Lambda([CacheMisses = ZenStats.CacheStats.Misses]
+			{
+				return FText::AsNumber(CacheMisses);
+			})
+		];
 
-	Panel->AddSlot(4, Row)
-	[
-		SNew(STextBlock)
-		.Margin(FMargin(ColumnMargin, TitleMargin, ColumnMargin, 0.0f))
-		.ColorAndOpacity(TitleColor)
-		.Font(TitleFont)
-		.Text(LOCTEXT("URL", "URL"))
-	];
 	Row++;
-	
-	int32 EndpointIndex = 1;
 
-	for (const UE::Zen::FZenEndPointStats& EndpointStats : ZenStats.UpstreamStats.EndPointStats)
+	if (!ZenStats.UpstreamStats.EndPointStats.IsEmpty())
 	{
 		Panel->AddSlot(0, Row)
 		[
 			SNew(STextBlock)
-			.Margin(FMargin(ColumnMargin, RowMargin))
-			.Text_Lambda([EndpointStats] { return FText::FromString(EndpointStats.Name); })
+			.Margin(FMargin(ColumnMargin, TitleMargin, ColumnMargin, 0.0f))
+			.ColorAndOpacity(TitleColor)
+			.Font(TitleFont)
+			.Text(LOCTEXT("UpstreamServer", "Upstream Server"))
 		];
 
 		Panel->AddSlot(1, Row)
 		[
 			SNew(STextBlock)
-			.Margin(FMargin(ColumnMargin, RowMargin))
-			.Text_Lambda([EndpointStats] { return FText::FromString(SingleDecimalFormat(EndpointStats.HitRatio * 100.0) + TEXT(" %")); })
+			.Margin(FMargin(ColumnMargin, TitleMargin, ColumnMargin, 0.0f))
+			.ColorAndOpacity(TitleColor)
+			.Font(TitleFont)
+			.Text(LOCTEXT("HitRate", "Hit Rate"))
 		];
 
 		Panel->AddSlot(2, Row)
 		[
 			SNew(STextBlock)
-			.Margin(FMargin(ColumnMargin, RowMargin))
-			.Text_Lambda([EndpointStats] 
-				{ 
-					if (EndpointStats.DownloadedMB> 1024.0)
-					{
-						return FText::FromString(SingleDecimalFormat(FUnitConversion::Convert(EndpointStats.DownloadedMB, EUnit::Megabytes, EUnit::Gigabytes) ) + TEXT(" GB")); 
-					}
-					else
-					{
-						return FText::FromString(SingleDecimalFormat(EndpointStats.DownloadedMB) + TEXT(" MB"));
-					}
-				})
-											
+			.Margin(FMargin(ColumnMargin, TitleMargin, ColumnMargin, 0.0f))
+			.ColorAndOpacity(TitleColor)
+			.Font(TitleFont)
+			.Text(LOCTEXT("Downloaded", "Downloaded"))
 		];
 
 		Panel->AddSlot(3, Row)
 		[
 			SNew(STextBlock)
-			.Margin(FMargin(ColumnMargin, RowMargin))
-			.Text_Lambda([EndpointStats]
-				{ 
-					if (EndpointStats.UploadedMB > 1024.0)
-					{
-						return FText::FromString(SingleDecimalFormat(FUnitConversion::Convert(EndpointStats.UploadedMB, EUnit::Megabytes, EUnit::Gigabytes) ) + TEXT(" GB")); 
-					}
-					else
-					{
-						return FText::FromString(SingleDecimalFormat(EndpointStats.UploadedMB) + TEXT(" MB"));
-					}
-				})
+			.Margin(FMargin(ColumnMargin, TitleMargin, ColumnMargin, 0.0f))
+			.ColorAndOpacity(TitleColor)
+			.Font(TitleFont)
+			.Text(LOCTEXT("Uploaded", "Uploaded"))
 		];
 
 		Panel->AddSlot(4, Row)
 		[
 			SNew(STextBlock)
-			.Margin(FMargin(ColumnMargin, RowMargin))
-			.Text_Lambda([EndpointStats] { return FText::FromString(EndpointStats.Url); })
+			.Margin(FMargin(ColumnMargin, TitleMargin, ColumnMargin, 0.0f))
+			.ColorAndOpacity(TitleColor)
+			.Font(TitleFont)
+			.Text(LOCTEXT("URL", "URL"))
 		];
-
-		SumTotalGetMB += EndpointStats.DownloadedMB;
-		SumTotalPutMB += EndpointStats.UploadedMB;
-
 		Row++;
-	}
 
-#endif
+		for (const UE::Zen::FZenEndPointStats& EndpointStats : ZenStats.UpstreamStats.EndPointStats)
+		{
+			Panel->AddSlot(0, Row)
+			[
+				SNew(STextBlock)
+				.Margin(FMargin(ColumnMargin, RowMargin))
+				.Text_Lambda([EndpointStats] { return FText::FromString(EndpointStats.Name); })
+			];
+
+			Panel->AddSlot(1, Row)
+			[
+				SNew(STextBlock)
+				.Margin(FMargin(ColumnMargin, RowMargin))
+				.Text_Lambda([EndpointStats] { return FText::AsPercent(EndpointStats.HitRatio, &SingleDecimalFormatting); })
+			];
+
+			Panel->AddSlot(2, Row)
+			[
+				SNew(STextBlock)
+				.Margin(FMargin(ColumnMargin, RowMargin))
+				.Text_Lambda([DownloadedMB = EndpointStats.DownloadedMB]
+					{
+						uint64 DownloadedBytes = DownloadedMB*1024*1024;
+						return FText::AsMemory(DownloadedBytes, (DownloadedBytes > 1024) ? &SingleDecimalFormatting : nullptr, nullptr, EMemoryUnitStandard::IEC);
+					})
+												
+			];
+
+			Panel->AddSlot(3, Row)
+			[
+				SNew(STextBlock)
+				.Margin(FMargin(ColumnMargin, RowMargin))
+				.Text_Lambda([UploadedMB = EndpointStats.UploadedMB]
+					{
+						uint64 UploadedBytes = UploadedMB*1024*1024;
+						return FText::AsMemory(UploadedBytes, (UploadedBytes > 1024) ? &SingleDecimalFormatting : nullptr, nullptr, EMemoryUnitStandard::IEC);
+					})
+			];
+
+			Panel->AddSlot(4, Row)
+			[
+				SNew(STextBlock)
+				.Margin(FMargin(ColumnMargin, RowMargin))
+				.Text_Lambda([EndpointStats] { return FText::FromString(EndpointStats.Url); })
+			];
+
+			Row++;
+		}
+	}
 
 	return Panel;
 }
diff --git a/Engine/Source/Programs/ZenDashboard/Private/SZenCacheStatistics.h b/Engine/Source/Programs/ZenDashboard/Private/SZenCacheStatistics.h
--- a/Engine/Source/Programs/ZenDashboard/Private/SZenCacheStatistics.h
+++ b/Engine/Source/Programs/ZenDashboard/Private/SZenCacheStatistics.h
@@ -2,18 +2,25 @@
 
 #pragma once
 
-#include "CoreMinimal.h"
+#include "ServiceInstanceManager.h"
+#include "Templates/SharedPointer.h"
 #include "Widgets/DeclarativeSyntaxSupport.h"
-#include "Widgets/SWidget.h"
-#include "Widgets/SCompoundWidget.h"
 #include "Widgets/SBoxPanel.h"
+#include "Widgets/SCompoundWidget.h"
+#include "Widgets/SWidget.h"
+#include "ZenServerInterface.h"
 
-class SZenCacheStatisticsDialog : public SCompoundWidget
+class SZenCacheStatistics : public SCompoundWidget
 {
-	SLATE_BEGIN_ARGS(SZenCacheStatisticsDialog) {}
+	SLATE_BEGIN_ARGS(SZenCacheStatistics)
+		: _ZenServiceInstance(nullptr)
+	{ }
+
+	SLATE_ATTRIBUTE(TSharedPtr<UE::Zen::FZenServiceInstance>, ZenServiceInstance);
+
 	SLATE_END_ARGS()
 
-		void Construct(const FArguments& InArgs);
+	void Construct(const FArguments& InArgs);
 
 private:
 
@@ -22,4 +29,5 @@ private:
 	EActiveTimerReturnType UpdateGridPanels(double InCurrentTime, float InDeltaTime);
 
 	SVerticalBox::FSlot* GridSlot = nullptr;
+	TAttribute<TSharedPtr<UE::Zen::FZenServiceInstance>> ZenServiceInstance;
 };
\ No newline at end of file
diff --git a/Engine/Source/Programs/ZenDashboard/Private/SZenServiceStatus.cpp b/Engine/Source/Programs/ZenDashboard/Private/SZenServiceStatus.cpp
new file mode 100644
--- /dev/null
+++ b/Engine/Source/Programs/ZenDashboard/Private/SZenServiceStatus.cpp
@@ -0,0 +1,246 @@
+// Copyright Epic Games, Inc. All Rights Reserved.
+
+#include "SZenServiceStatus.h"
+#include "ZenServerInterface.h"
+#include "Math/UnitConversion.h"
+#include "Math/BasicMathExpressionEvaluator.h"
+#include "Misc/ExpressionParser.h"
+#include "Styling/StyleColors.h"
+#include "Widgets/Layout/SGridPanel.h"
+#include "Widgets/Text/STextBlock.h"
+#include "Widgets/SBoxPanel.h"
+#include "Internationalization/FastDecimalFormat.h"
+#include "Widgets/Input/SButton.h"
+#include "Widgets/Images/SImage.h"
+
+#define LOCTEXT_NAMESPACE "ZenDashboard"
+
+void SZenServiceStatus::Construct(const FArguments& InArgs)
+{
+	ZenServiceInstance = InArgs._ZenServiceInstance;
+
+	this->ChildSlot
+	[
+		SNew(SVerticalBox)
+
+		+ SVerticalBox::Slot()
+		.AutoHeight()
+		.Padding(0, 0, 0, 0)
+		.Expose(GridSlot)
+		[
+			GetGridPanel()
+		]
+	];
+
+	UpdateState(0.0, 0.0);
+	RegisterActiveTimer(0.5f, FWidgetActiveTimerDelegate::CreateSP(this, &SZenServiceStatus::UpdateState));
+}
+
+EActiveTimerReturnType SZenServiceStatus::UpdateState(double InCurrentTime, float InDeltaTime)
+{
+	uint32 LastStateIndex = ActiveStateIndex.load(std::memory_order_relaxed);
+	uint32 NextStateIndex = (LastStateIndex + 1)%NumState;
+	FState& NextState = State[NextStateIndex];
+	NextState.bGotRunContext = UE::Zen::TryGetLocalServiceRunContext(NextState.RunContext);
+	NextState.Version = UE::Zen::GetLocalServiceInstallVersion(false);
+	if (NextState.bGotRunContext)
+	{
+		NextState.bIsRunning = UE::Zen::IsLocalServiceRunning(*NextState.RunContext.GetDataPath(), &NextState.LocalPort);
+	}
+
+	if (NextState.bIsRunning)
+	{
+		if (TSharedPtr<UE::Zen::FZenServiceInstance> ServiceInstance = ZenServiceInstance.Get())
+		{
+			ServiceInstance->GetGCStatus(NextState.GCStatus);
+		}
+	}
+
+	ActiveStateIndex.compare_exchange_weak(LastStateIndex, NextStateIndex, std::memory_order_relaxed);
+
+	return EActiveTimerReturnType::Continue;
+}
+
+FReply SZenServiceStatus::ExploreDataPath_OnClicked()
+{
+	const FState& CurrentState = GetCurrentState();
+	if (CurrentState.bGotRunContext)
+	{
+		FString FullPath(FPaths::ConvertRelativePathToFull(CurrentState.RunContext.GetDataPath()));
+		FPlatformProcess::ExploreFolder(*FullPath);
+		return FReply::Handled();
+	}
+	return FReply::Unhandled();
+}
+
+const SZenServiceStatus::FState& SZenServiceStatus::GetCurrentState() const
+{
+	return State[ActiveStateIndex.load(std::memory_order_relaxed)];
+}
+
+TSharedRef<SWidget> SZenServiceStatus::GetGridPanel()
+{
+	TSharedRef<SGridPanel> Panel = SNew(SGridPanel);
+
+	int32 Row = 0;
+
+	const float RowMargin = 0.0f;
+	const float ColumnMargin = 10.0f;
+	const FSlateColor TitleColor = FStyleColors::AccentWhite;
+	const FSlateFontInfo TitleFont = FCoreStyle::GetDefaultFontStyle("Bold", 10);
+
+	Panel->AddSlot(0, Row)
+	[
+		SNew(STextBlock)
+		.Margin(FMargin(ColumnMargin, RowMargin))
+		.ColorAndOpacity(TitleColor)
+		.Font(TitleFont)
+		.Text(LOCTEXT("ServiceStatus_Status", "Status"))
+	];
+
+	Panel->AddSlot(1, Row)
+	[
+		SNew(STextBlock)
+		.Margin(FMargin(ColumnMargin, RowMargin))
+		.Text_Lambda([this]
+		{
+			const FState& CurrentState = GetCurrentState();
+			FText ServiceStatus = LOCTEXT("ServiceStatus_NotInstalled", "Not installed");
+			if (CurrentState.bGotRunContext)
+			{
+				if (CurrentState.bIsRunning && (CurrentState.LocalPort != 0))
+				{
+					ServiceStatus = LOCTEXT("ServiceStatus_Running", "Running");
+				}
+				else
+				{
+					ServiceStatus = LOCTEXT("ServiceStatus_Stopped", "Stopped");
+				}
+			}
+			return ServiceStatus;
+		})
+	];
+
+	++Row;
+
+	Panel->AddSlot(0, Row)
+	[
+		SNew(STextBlock)
+		.Margin(FMargin(ColumnMargin, RowMargin))
+		.ColorAndOpacity(TitleColor)
+		.Font(TitleFont)
+		.Text(LOCTEXT("ServiceStatus_Version", "Version"))
+	];
+
+	Panel->AddSlot(1, Row)
+	[
+		SNew(STextBlock)
+		.Margin(FMargin(ColumnMargin, RowMargin))
+		.Text_Lambda([this]
+		{
+			const FState& CurrentState = GetCurrentState();
+			return FText::FromString(CurrentState.Version);
+		})
+	];
+
+	++Row;
+
+	Panel->AddSlot(0, Row)
+	[
+		SNew(STextBlock)
+		.Margin(FMargin(ColumnMargin, RowMargin))
+		.ColorAndOpacity(TitleColor)
+		.Font(TitleFont)
+		.Text(LOCTEXT("ServiceStatus_Port", "Port"))
+	];
+
+	Panel->AddSlot(1, Row)
+	[
+		SNew(STextBlock)
+		.Margin(FMargin(ColumnMargin, RowMargin))
+		.Text_Lambda([this]
+		{
+			const FState& CurrentState = GetCurrentState();
+			return CurrentState.LocalPort == 0 ? LOCTEXT("ServiceStatus_NoPortValue", "-") : FText::AsNumber(CurrentState.LocalPort, &FNumberFormattingOptions::DefaultNoGrouping());
+		})
+	];
+
+	++Row;
+
+	Panel->AddSlot(0, Row)
+	[
+		SNew(STextBlock)
+		.Margin(FMargin(ColumnMargin, RowMargin))
+		.ColorAndOpacity(TitleColor)
+		.Font(TitleFont)
+		.Text(LOCTEXT("ServiceStatus_GCStatus", "GC Status"))
+	];
+
+	Panel->AddSlot(1, Row)
+	[
+		SNew(STextBlock)
+		.Margin(FMargin(ColumnMargin, RowMargin))
+		.Text_Lambda([this]
+		{
+			const FState& CurrentState = GetCurrentState();
+			return CurrentState.LocalPort == 0 ? LOCTEXT("ServiceStatus_NoGCStatusValue", "-") : FText::FromString(CurrentState.GCStatus.Description);
+		})
+	];
+
+	++Row;
+
+	Panel->AddSlot(0, Row)
+	.VAlign(VAlign_Bottom)
+	[
+		SNew(STextBlock)
+		.Margin(FMargin(ColumnMargin, RowMargin))
+		.ColorAndOpacity(TitleColor)
+		.Font(TitleFont)
+		.Text(LOCTEXT("ServiceStatus_DataPath", "Data path"))
+	];
+
+	Panel->AddSlot(1, Row)
+	[
+
+		SNew(SHorizontalBox)
+
+		+ SHorizontalBox::Slot()
+		.VAlign(VAlign_Bottom)
+		[
+			SNew(STextBlock)
+			.Margin(FMargin(ColumnMargin, RowMargin))
+			.Text_Lambda([this]
+			{
+				const FState& CurrentState = GetCurrentState();
+				return CurrentState.bGotRunContext ? FText::FromString(CurrentState.RunContext.GetDataPath()) : LOCTEXT("ServiceStatus_NoDataPathValue", "-");
+			})
+		]
+
+		+ SHorizontalBox::Slot()
+		.AutoWidth()
+		.VAlign(VAlign_Bottom)
+		[
+			SNew(SButton)
+			.ButtonStyle(&FAppStyle::Get().GetWidgetStyle<FButtonStyle>("SimpleButton"))
+			.ToolTipText(LOCTEXT("ExploreDataPathButtonToolTip", "Explore the Data Path"))
+			.IsEnabled_Lambda([this]
+			{
+				const FState& CurrentState = GetCurrentState();
+				return CurrentState.bGotRunContext;
+			})
+			.OnClicked(this, &SZenServiceStatus::ExploreDataPath_OnClicked)
+			[
+				SNew(SImage)
+				.Image(FAppStyle::Get().GetBrush("Icons.FolderExplore"))
+				.ColorAndOpacity(FSlateColor::UseForeground())
+			]
+		]
+	];
+
+	++Row;
+
+
+	return Panel;
+}
+
+#undef LOCTEXT_NAMESPACE
diff --git a/Engine/Source/Programs/ZenDashboard/Private/SZenServiceStatus.h b/Engine/Source/Programs/ZenDashboard/Private/SZenServiceStatus.h
new file mode 100644
--- /dev/null
+++ b/Engine/Source/Programs/ZenDashboard/Private/SZenServiceStatus.h
@@ -0,0 +1,48 @@
+// Copyright Epic Games, Inc. All Rights Reserved.
+
+#pragma once
+
+#include "ServiceInstanceManager.h"
+#include "Templates/SharedPointer.h"
+#include "Widgets/DeclarativeSyntaxSupport.h"
+#include "Widgets/SBoxPanel.h"
+#include "Widgets/SCompoundWidget.h"
+#include "Widgets/SWidget.h"
+#include "ZenServerInterface.h"
+
+class SZenServiceStatus : public SCompoundWidget
+{
+	SLATE_BEGIN_ARGS(SZenServiceStatus)
+		: _ZenServiceInstance(nullptr)
+	{ }
+
+	SLATE_ATTRIBUTE(TSharedPtr<UE::Zen::FZenServiceInstance>, ZenServiceInstance);
+
+	SLATE_END_ARGS()
+
+	void Construct(const FArguments& InArgs);
+
+private:
+	struct FState
+	{
+		UE::Zen::FZenLocalServiceRunContext RunContext;
+		FString Version;
+		UE::Zen::FGCStatus GCStatus;
+		uint16 LocalPort = 0;
+		bool bGotRunContext = false;
+		bool bIsRunning = false;
+	};
+	static constexpr uint32 NumState = 2;
+	FState State[NumState];
+	std::atomic<uint32> ActiveStateIndex = 0;
+
+
+	const FState& GetCurrentState() const;
+	TSharedRef<SWidget> GetGridPanel();
+	FReply ExploreDataPath_OnClicked();
+
+	EActiveTimerReturnType UpdateState(double InCurrentTime, float InDeltaTime);
+
+	SVerticalBox::FSlot* GridSlot = nullptr;
+	TAttribute<TSharedPtr<UE::Zen::FZenServiceInstance>> ZenServiceInstance;
+};
\ No newline at end of file
diff --git a/Engine/Source/Programs/ZenDashboard/Private/ServiceInstanceManager.cpp b/Engine/Source/Programs/ZenDashboard/Private/ServiceInstanceManager.cpp
new file mode 100644
--- /dev/null
+++ b/Engine/Source/Programs/ZenDashboard/Private/ServiceInstanceManager.cpp
@@ -0,0 +1,35 @@
+// Copyright Epic Games, Inc. All Rights Reserved.
+
+#include "ServiceInstanceManager.h"
+#include "ZenServerInterface.h"
+
+
+namespace UE::Zen
+{
+
+TSharedPtr<FZenServiceInstance> FServiceInstanceManager::GetZenServiceInstance() const
+{
+	UE::Zen::FZenLocalServiceRunContext RunContext;
+	if (UE::Zen::TryGetLocalServiceRunContext(RunContext))
+	{
+		uint16 LocalPort = 0;
+		if (UE::Zen::IsLocalServiceRunning(*RunContext.GetDataPath(), &LocalPort) && (LocalPort != 0))
+		{
+			if (LocalPort != CurrentPort)
+			{
+				CurrentPort = LocalPort;
+				UE::Zen::FServiceSettings ServiceSettings;
+				ServiceSettings.SettingsVariant.Emplace<UE::Zen::FServiceConnectSettings>();
+				UE::Zen::FServiceConnectSettings& ConnectExistingSettings = ServiceSettings.SettingsVariant.Get<UE::Zen::FServiceConnectSettings>();
+				ConnectExistingSettings.HostName = TEXT("localhost");
+				ConnectExistingSettings.Port = LocalPort;
+				CurrentInstance = MakeShared<FZenServiceInstance>(MoveTemp(ServiceSettings));
+			}
+
+			return CurrentInstance;
+		}
+	}
+	return nullptr;
+}
+
+} // namespace UE::Zen
diff --git a/Engine/Source/Programs/ZenDashboard/Private/ServiceInstanceManager.h b/Engine/Source/Programs/ZenDashboard/Private/ServiceInstanceManager.h
new file mode 100644
--- /dev/null
+++ b/Engine/Source/Programs/ZenDashboard/Private/ServiceInstanceManager.h
@@ -0,0 +1,22 @@
+// Copyright Epic Games, Inc. All Rights Reserved.
+
+#pragma once
+
+#include "Templates/SharedPointer.h"
+
+namespace UE::Zen { class FZenServiceInstance; }
+
+namespace UE::Zen
+{
+
+class FServiceInstanceManager
+{
+public:
+	TSharedPtr<FZenServiceInstance> GetZenServiceInstance() const;
+private:
+	mutable uint16 CurrentPort = 0;
+	mutable TSharedPtr<FZenServiceInstance> CurrentInstance;
+};
+
+} // namespace UE::Zen
+
diff --git a/Engine/Source/Programs/ZenDashboard/Private/ZenDashboard.cpp b/Engine/Source/Programs/ZenDashboard/Private/ZenDashboard.cpp
--- a/Engine/Source/Programs/ZenDashboard/Private/ZenDashboard.cpp
+++ b/Engine/Source/Programs/ZenDashboard/Private/ZenDashboard.cpp
@@ -1,34 +1,46 @@
 // Copyright Epic Games, Inc. All Rights Reserved.
 
 #include "ZenDashboard.h"
-#include "RequiredProgramMainCPPInclude.h"
+#include "Features/IModularFeatures.h"
 #include "Framework/Application/SlateApplication.h"
-#include "StandaloneRenderer.h"
-#include "ZenDashboardStyle.h"
+#include "Framework/MultiBox/MultiBoxBuilder.h"
+#include "Framework/Notifications/NotificationManager.h"
+#include "GenericPlatform/GenericPlatformFile.h"
 #include "HAL/PlatformProcess.h"
+#include "Misc/CompilationResult.h"
+#include "Misc/MessageDialog.h"
+#include "Misc/Optional.h"
+#include "Modules/ModuleInterface.h"
 #include "Modules/ModuleManager.h"
-#include "Features/IModularFeatures.h"
-#include "Widgets/Notifications/SNotificationList.h"
+#include "RequiredProgramMainCPPInclude.h"
+#include "ServiceInstanceManager.h"
+#include "StandaloneRenderer.h"
+#include "SZenCacheStatistics.h"
+#include "SZenServiceStatus.h"
+#include "Templates/SharedPointer.h"
 #include "Widgets/Input/SButton.h"
-#include "Framework/Notifications/NotificationManager.h"
-#include "Windows/WindowsHWrapper.h"
-#include "Misc/MonitoredProcess.h"
-#include "Modules/ModuleInterface.h"
 #include "Widgets/Input/SButton.h"
 #include "Widgets/Input/SCheckBox.h"
+#include "Widgets/Notifications/SNotificationList.h"
 #include "Widgets/Text/STextBlock.h"
-#include "Misc/CompilationResult.h"
-#include "Misc/MessageDialog.h"
-#include "SZenCacheStatistics.h"
+#include "ZenDashboardStyle.h"
 #include "ZenServerInterface.h"
 
 #if PLATFORM_WINDOWS
+#include "Windows/WindowsHWrapper.h"
 #include "Runtime/Launch/Resources/Windows/Resource.h"
 #include "Windows/WindowsApplication.h"
 #include "Windows/WindowsHWrapper.h"
 #include "Windows/PreWindowsApi.h"
 #include <shellapi.h>
 #include "Windows/PostWindowsApi.h"
+#elif PLATFORM_LINUX
+#include "UnixCommonStartup.h"
+#elif PLATFORM_MAC
+#include "Mac/MacProgramDelegate.h"
+#include "LaunchEngineLoop.h"
+#else
+#error "Unsupported platform!"
 #endif
 
 #define LOCTEXT_NAMESPACE "ZenDashboard"
@@ -40,8 +52,20 @@ static void OnRequestExit()
 	RequestEngineExit(TEXT("ZenDashboard Closed"));
 }
 
+static void HideOnCloseOverride(const TSharedRef<SWindow>& WindowBeingClosed)
+{
+	WindowBeingClosed->HideWindow();
+}
+
 #if PLATFORM_WINDOWS
 #define WM_TRAYICON (WM_USER + 1000)
+
+#define CMD_USER 400
+#define CMD_SHOWDASHBOARD					(CMD_USER +   1)
+#define CMD_STARTZENSERVER					(CMD_USER +   2)
+#define CMD_STOPZENSERVER					(CMD_USER +   3)
+#define CMD_RESTARTZENSERVER				(CMD_USER +   4)
+#define CMD_EXITDASHBOARD					(CMD_USER +   5)
 #endif
 
 class FZenDashboardApp
@@ -49,86 +73,373 @@ class FZenDashboardApp
 	: IWindowsMessageHandler
 #endif
 {
-private:
 	FCriticalSection CriticalSection;
 	FSlateApplication& Slate;
 	TSharedPtr<SWindow> Window;
 	TSharedPtr<SNotificationItem> CompileNotification;
+	TSharedPtr<UE::Zen::FServiceInstanceManager> ServiceInstanceManager;
 	TArray<FSimpleDelegate> MainThreadTasks;
-	bool bRequestCancel;
-	bool bDisableLimit;
-	bool bHasReinstancingProcess;
-	bool bWarnOnRestart;
-	FDateTime LastPatchTime;
-	FDateTime NextPatchStartTime;
-	UE::Zen::FScopeZenService ZenService;
-	UE::Zen::FZenStats ZenStats;
-
 
+#if PLATFORM_WINDOWS
 	bool ProcessMessage(HWND hwnd, uint32 msg, WPARAM wParam, LPARAM lParam, int32& OutResult) override
 	{
-		if (msg == WM_TRAYICON && LOWORD(lParam) == WM_LBUTTONDBLCLK)
+		switch (msg)
 		{
-			Window->ShowWindow();
-			return true;
+		case WM_TRAYICON:
+			if (LOWORD(lParam) == WM_LBUTTONDBLCLK)
+			{
+				Window->ShowWindow();
+				return true;
+			}
+			else if (LOWORD(lParam) == WM_RBUTTONUP)
+			{
+				HMENU Menu = CreatePopupMenu();
+
+				auto AddMenuItem = [Menu](const wchar_t* Text, uint32 Id)
+				{
+					MENUITEMINFOW MenuItemInfo = {
+						sizeof(MenuItemInfo),
+						MIIM_TYPE|MIIM_STATE|MIIM_ID,
+						MFT_STRING,
+						MFS_ENABLED,
+						Id,
+					};
+					MenuItemInfo.dwTypeData = LPWSTR(Text);
+					InsertMenuItemW(Menu, 999, true, &MenuItemInfo);
+				};
+
+				AddMenuItem(*LOCTEXT("Show_ZenDashboard", "Zen Dashboard").ToString(), CMD_SHOWDASHBOARD);
+				AddMenuItem(*LOCTEXT("Start_ZenServer", "Start Zen Server").ToString(), CMD_STARTZENSERVER);
+				AddMenuItem(*LOCTEXT("Stop_ZenServer", "Stop Zen Server").ToString(), CMD_STOPZENSERVER);
+				AddMenuItem(*LOCTEXT("Restart_ZenServer", "Restart Zen Server").ToString(), CMD_RESTARTZENSERVER);
+				AddMenuItem(*LOCTEXT("Exit", "Exit").ToString(), CMD_EXITDASHBOARD);
+
+				POINT MousePos;
+				GetCursorPos(&MousePos);
+				BOOL Ret = TrackPopupMenu(
+					Menu, TPM_RIGHTALIGN|TPM_BOTTOMALIGN,
+					MousePos.x, MousePos.y,
+					0, hwnd, nullptr);
+
+				DestroyMenu(Menu);
+			}
+			break;
+		case WM_COMMAND:
+			switch (LOWORD(wParam))
+			{
+			case CMD_SHOWDASHBOARD:
+				Window->ShowWindow();
+				break;
+			case CMD_STARTZENSERVER:
+				StartZenServer();
+				break;
+			case CMD_STOPZENSERVER:
+				StopZenServer();
+				break;
+			case CMD_RESTARTZENSERVER:
+				RestartZenServer();
+				break;
+			case CMD_EXITDASHBOARD:
+				ExitDashboard();
+				break;
+			}
+			break;
 		}
+
 		return false;
 	}
+#endif // PLATFORM_WINDOWS
+
+	void ExitDashboard()
+	{
+		FSlateApplication::Get().RequestDestroyWindow(Window.ToSharedRef());
+	}
+
+	void StartZenServer()
+	{
+		UE::Zen::FZenLocalServiceRunContext RunContext;
+		if (UE::Zen::TryGetLocalServiceRunContext(RunContext))
+		{
+			UE::Zen::StartLocalService(RunContext);
+		}
+	}
+
+	void StopZenServer()
+	{
+		UE::Zen::FZenLocalServiceRunContext RunContext;
+		if (UE::Zen::TryGetLocalServiceRunContext(RunContext))
+		{
+			UE::Zen::StopLocalService(*RunContext.GetDataPath());
+		}
+	}
+
+	void RestartZenServer()
+	{
+		StopZenServer();
+		StartZenServer();
+	}
+
+	void DeleteDataAndRestartZenServer()
+	{
+		UE::Zen::FZenLocalServiceRunContext RunContext;
+		if (UE::Zen::TryGetLocalServiceRunContext(RunContext))
+		{
+			if (UE::Zen::IsLocalServiceRunning(*RunContext.GetDataPath()))
+			{
+				FText Title = LOCTEXT("ConfirmToDeleteZenServerData_Title", "Zen Server Data Delete");
+				FText ConfirmMessage = FText::Format(LOCTEXT("ConfirmToDeleteZenServerData", "You are about to delete all local Zen Server data at:\n\n{0}\n\nThis will require you to build or download all data stored there which can take a long time. Are you sure you want to proceed?"),
+					FText::FromString(RunContext.GetDataPath()));
+				EAppReturnType::Type OkToDelete = OpenModalMessageDialog_Internal(EAppMsgType::YesNo, ConfirmMessage, Title, Window);
+				if (OkToDelete == EAppReturnType::No)
+				{
+					return;
+				}
+
+				StopZenServer();
+
+				FPlatformFileManager::Get().GetPlatformFile().DeleteDirectoryRecursively(*RunContext.GetDataPath());
+
+				StartZenServer();
+			}
+		}
+	}
+
+	void GatherDiagnostics()
+	{
+	}
+
+	void LaunchUtilityCommandPrompt()
+	{
+#if PLATFORM_WINDOWS
+		TStringBuilder<128> Args;
+		FString UtilityPath = UE::Zen::GetLocalInstallUtilityPath();
+		Args << TEXT("/k \"") << UtilityPath << TEXT("\" --help");
+		FProcHandle UtilProcHandle = FPlatformProcess::CreateProc(TEXT("cmd.exe"), Args.ToString(), false, false, false, nullptr, 0, *FPaths::GetPath(UtilityPath), nullptr);
+#else
+		// TODO: Need to get the equivalent for linux/mac
+#endif
+	}
+
+	void RunGC()
+	{
+		if (TSharedPtr<UE::Zen::FZenServiceInstance> ZenServiceInstance = ServiceInstanceManager->GetZenServiceInstance())
+		{
+			ZenServiceInstance->RequestGC();
+		}
+	}
+
+	TSharedRef< SWidget > MakeMainMenu()
+	{
+		FMenuBarBuilder MenuBuilder( nullptr );
+		{
+			// Control
+			MenuBuilder.AddPullDownMenu(
+				LOCTEXT( "FileMenu", "File" ),
+				LOCTEXT( "FileMenu_ToolTip", "Opens the file menu" ),
+				FNewMenuDelegate::CreateRaw( this, &FZenDashboardApp::FillFileMenu ) );
+
+			// Control
+			MenuBuilder.AddPullDownMenu(
+				LOCTEXT( "ToolsMenu", "Tools" ),
+				LOCTEXT( "ToolsMenu_ToolTip", "Opens the tools menu" ),
+				FNewMenuDelegate::CreateRaw( this, &FZenDashboardApp::FillToolsMenu ) );
+		}
+
+		// Create the menu bar
+		TSharedRef< SWidget > MenuBarWidget = MenuBuilder.MakeWidget();
+		MenuBarWidget->SetVisibility( EVisibility::Visible ); // Work around for menu bar not showing on Mac
+
+		return MenuBarWidget;
+	}
+
+	void FillFileMenu(FMenuBuilder& MenuBuilder)
+	{
+		MenuBuilder.AddMenuEntry(
+			LOCTEXT("Exit", "Exit"),
+			LOCTEXT("Exit_ToolTip", "Exits the Zen Dashboard"),
+			FSlateIcon(),
+			FUIAction(
+				FExecuteAction::CreateRaw( this, &FZenDashboardApp::ExitDashboard ),
+				FCanExecuteAction()
+			),
+			NAME_None,
+			EUserInterfaceActionType::Button
+		);
+	}
+
+	void FillToolsMenu(FMenuBuilder& MenuBuilder)
+	{
+		MenuBuilder.BeginSection(NAME_None, LOCTEXT("ServiceControls", "Service controls"));
+		{
+			MenuBuilder.AddMenuEntry(
+				LOCTEXT("Start_ZenServer", "Start Zen Server"),
+				LOCTEXT("Start_ZenServer_ToolTip", "Ensures ZenServer is started"),
+				FSlateIcon(),
+				FUIAction(
+					FExecuteAction::CreateRaw(this, &FZenDashboardApp::StartZenServer),
+					FCanExecuteAction()
+				),
+				NAME_None,
+				EUserInterfaceActionType::Button
+			);
+			MenuBuilder.AddMenuEntry(
+				LOCTEXT("Stop_ZenServer", "Stop Zen Server"),
+				LOCTEXT("Stop_ZenServer_ToolTip", "Ensures ZenServer is stopped"),
+				FSlateIcon(),
+				FUIAction(
+					FExecuteAction::CreateRaw(this, &FZenDashboardApp::StopZenServer),
+					FCanExecuteAction()
+				),
+				NAME_None,
+				EUserInterfaceActionType::Button
+			);
+			MenuBuilder.AddMenuEntry(
+				LOCTEXT("Restart_ZenServer", "Restart Zen Server"),
+				LOCTEXT("Restart_ZenServer_ToolTip", "Stops ZenServer if it is running then ensures it is running"),
+				FSlateIcon(),
+				FUIAction(
+					FExecuteAction::CreateRaw(this, &FZenDashboardApp::RestartZenServer),
+					FCanExecuteAction()
+				),
+				NAME_None,
+				EUserInterfaceActionType::Button
+			);
+
+			MenuBuilder.AddSubMenu(LOCTEXT("Advanced", "Advanced"), LOCTEXT("Advanced_ToolTip", "Advanced service control commands"),
+				FNewMenuDelegate::CreateLambda([this](FMenuBuilder& SubMenuBuilder)
+					{
+						SubMenuBuilder.AddMenuEntry(
+							LOCTEXT("DeleteDataAndRestart_ZenServer", "Delete data and restart Zen Server"),
+							LOCTEXT("DeleteDataAndRestart_ZenServer_ToolTip", "Stops ZenServer if it is running, deletes ALL of its data, then starts it from a blank state"),
+							FSlateIcon(),
+							FUIAction(
+								FExecuteAction::CreateRaw(this, &FZenDashboardApp::DeleteDataAndRestartZenServer),
+								FCanExecuteAction()
+							),
+							NAME_None,
+							EUserInterfaceActionType::Button
+						);
+					})
+			);
+		}
+		MenuBuilder.EndSection();
+
+		MenuBuilder.BeginSection(NAME_None, LOCTEXT("Debug", "Debug"));
+		{
+			// TODO: The "Gather diagnostic info" feature is off for now and needs to be filled in as part of a subsequent change
+			// MenuBuilder.AddMenuEntry(
+			// 	LOCTEXT("GatherDiagnosticInfo", "Gather diagnostic info"),
+			// 	LOCTEXT("GatherDiagnosticInfo_ToolTip", "Gathers information for diagnosing issues"),
+			// 	FSlateIcon(),
+			// 	FUIAction(
+			// 		FExecuteAction::CreateRaw(this, &FZenDashboardApp::GatherDiagnostics),
+			// 		FCanExecuteAction()
+			// 	),
+			// 	NAME_None,
+			// 	EUserInterfaceActionType::Button
+			// );
+			MenuBuilder.AddMenuEntry(
+				LOCTEXT("LaunchUtilityCommandPrompt", "Launch utility command prompt"),
+				LOCTEXT("LaunchUtilityCommandPrompt_ToolTip", "Launches command prompt for issuing zen server specific commands"),
+				FSlateIcon(),
+				FUIAction(
+					FExecuteAction::CreateRaw(this, &FZenDashboardApp::LaunchUtilityCommandPrompt),
+					FCanExecuteAction()
+				),
+				NAME_None,
+				EUserInterfaceActionType::Button
+			);
+		}
+		MenuBuilder.EndSection();
+
+		MenuBuilder.BeginSection(NAME_None, LOCTEXT("Disk", "Disk"));
+		{
+			MenuBuilder.AddMenuEntry(
+				LOCTEXT("FreeUnusedDiskSpace", "Free unused disk space"),
+				LOCTEXT("FreeUnusedDiskSpace_ToolTip", "Frees unused disk space by running a garbage collection process"),
+				FSlateIcon(),
+				FUIAction(
+					FExecuteAction::CreateRaw(this, &FZenDashboardApp::RunGC),
+					FCanExecuteAction()
+				),
+				NAME_None,
+				EUserInterfaceActionType::Button
+			);
+		}
+		MenuBuilder.EndSection();
+	}
 
 public:
 	FZenDashboardApp(FSlateApplication& InSlate)
 		: Slate(InSlate)
-		, bRequestCancel(false)
-		, bDisableLimit(false)
-		, bHasReinstancingProcess(false)
-		, bWarnOnRestart(false)
-		, LastPatchTime(FDateTime::MinValue())
-		, NextPatchStartTime(FDateTime::MinValue())
 	{
+		ServiceInstanceManager = MakeShared<UE::Zen::FServiceInstanceManager>();
+		InstallMessageDialogOverride();
 	}
 
 	virtual ~FZenDashboardApp()
 	{
+		RemoveMessageDialogOverride();
 	}
 
 	void Run()
 	{
+		const bool bSystemTrayMode = !!PLATFORM_WINDOWS;
+		
 		Window =
 			SNew(SWindow)
 			.Title(GetWindowTitle())
 			.ClientSize(FVector2D(384.0f, 128.0f))
 			.ActivationPolicy(EWindowActivationPolicy::Always)
 			.SizingRule(ESizingRule::Autosized)
-			.IsTopmostWindow(true)
+			.IsTopmostWindow(false)
 			.FocusWhenFirstShown(false)
 			.SupportsMaximize(false)
 			.SupportsMinimize(true)
 			.HasCloseButton(true)
 			[
 				SNew(SBorder)
-				.VAlign(VAlign_Center)
 				.BorderImage(FCoreStyle::Get().GetBrush("ToolPanel.GroupBorder"))
 				[
 					SNew(SVerticalBox)
+
+					// Menu
+					+ SVerticalBox::Slot()
+					.AutoHeight()
+					[
+						MakeMainMenu()
+					]
+
+					// Status panel
+					+ SVerticalBox::Slot()
+					.AutoHeight()
+					.Padding(0.0f, 10.0f, 0.0f, 0.0f)
+					.VAlign(VAlign_Top)
+					[
+						SNew(SZenServiceStatus)
+						.ZenServiceInstance(ServiceInstanceManager.ToSharedRef(), &UE::Zen::FServiceInstanceManager::GetZenServiceInstance)
+					]
+
+					// Stats panel
 					+ SVerticalBox::Slot()
 					.AutoHeight()
+					.Padding(0.0f, 10.0f, 0.0f, 0.0f)
 					.VAlign(VAlign_Top)
 					[
-						SNew(SZenCacheStatisticsDialog)
+						SNew(SZenCacheStatistics)
+						.ZenServiceInstance(ServiceInstanceManager.ToSharedRef(), &UE::Zen::FServiceInstanceManager::GetZenServiceInstance)
 					]
 				]
 			];
 
-		/*Window->RegisterActiveTimer(0.2f, FWidgetActiveTimerDelegate::CreateLambda([this](double InCurrentTime, float InDeltaTime)
-			{
-				ZenService.GetInstance().GetStats(ZenStats);
-				return EActiveTimerReturnType::Continue;
-			}));*/
+		if (bSystemTrayMode)
+		{
+			Window->SetRequestDestroyWindowOverride(FRequestDestroyWindowOverride::CreateStatic(&HideOnCloseOverride));
+		}
 
-		// Add the window without showing it
-		Slate.AddWindow(Window.ToSharedRef(), true);
+		Slate.AddWindow(Window.ToSharedRef(), !bSystemTrayMode);
 
-		// Show the window without stealling focus
+		// Show the window without stealing focus
 /*		if (!FParse::Param(FCommandLine::Get(), TEXT("Hidden")))
 		{
 			HWND ForegroundWindow = GetForegroundWindow();
@@ -194,58 +505,47 @@ private:
 		return LOCTEXT("WindowTitle", "Zen Dashboard");
 	}
 
-	void BringToFrontAsync()
-	{
-		FScopeLock Lock(&CriticalSection);
-		MainThreadTasks.Add(FSimpleDelegate::CreateRaw(this, &FZenDashboardApp::BringToFront));
-	}
-
-	void BringToFront()
-	{
-		HWND WindowHandle = (HWND)Window->GetNativeWindow()->GetOSWindowHandle();
-		if (IsIconic(WindowHandle))
-		{
-			ShowWindow(WindowHandle, SW_RESTORE);
-		}
-		::SetWindowPos(WindowHandle, HWND_TOPMOST, 0, 0, 0, 0, SWP_NOMOVE | SWP_NOSIZE | SWP_NOACTIVATE);
-		::SetWindowPos(WindowHandle, HWND_NOTOPMOST, 0, 0, 0, 0, SWP_NOMOVE | SWP_NOSIZE | SWP_NOACTIVATE);
-	}
-
-	void SetVisibleAsync(bool bVisible)
+	EAppReturnType::Type OnModalMessageDialog(EAppMsgType::Type InMessage, const FText& InText, const FText& InTitle)
 	{
-		FScopeLock Lock(&CriticalSection);
-		MainThreadTasks.Add(FSimpleDelegate::CreateLambda([this, bVisible](){ SetVisible(bVisible); }));
-	}
-
-	void SetVisible(bool bVisible)
-	{
-		if (bVisible)
+		if (IsInGameThread() && FSlateApplication::IsInitialized() && FSlateApplication::Get().CanAddModalWindow())
 		{
-			if (!Window->IsVisible())
-			{
-				Window->ShowWindow();
-			}
+			return OpenModalMessageDialog_Internal(InMessage, InText, InTitle, Window);
 		}
 		else
 		{
-			if (Window->IsVisible())
-			{
-				Window->HideWindow();
-			}
+			return FPlatformMisc::MessageBoxExt(InMessage, *InText.ToString(), *InTitle.ToString());
 		}
 	}
 
-	void ShowConsole()
+	void InstallMessageDialogOverride()
 	{
-		SetVisible(true);
-		BringToFront();
+		FCoreDelegates::ModalErrorMessage.BindRaw(this, &FZenDashboardApp::OnModalMessageDialog);
+	}
+
+	void RemoveMessageDialogOverride()
+	{
+		FCoreDelegates::ModalErrorMessage.Unbind();
 	}
 };
 
-bool ZenDashboardMain(const TCHAR* CmdLine)
+int ZenDashboardMain(const TCHAR* CmdLine)
 {
+	ON_SCOPE_EXIT
+	{
+		RequestEngineExit(TEXT("Exiting"));
+		FEngineLoop::AppPreExit();
+		FModuleManager::Get().UnloadModulesAtShutdown();
+		FEngineLoop::AppExit();
+	};
+
 	// start up the main loop
 	GEngineLoop.PreInit(CmdLine);
+
+	FSystemWideCriticalSection SystemWideZenDashboardCritSec(TEXT("ZenDashboard"));
+	if (!SystemWideZenDashboardCritSec.IsValid())
+	{
+		return true;
+	}
 	check(GConfig && GConfig->IsReadyForUse());
 
 	// Initialize high DPI mode
@@ -289,16 +589,28 @@ bool ZenDashboardMain(const TCHAR* CmdLine)
 		FSlateApplication::Shutdown();
 	}
 
-	FEngineLoop::AppPreExit();
-	FModuleManager::Get().UnloadModulesAtShutdown();
-	FEngineLoop::AppExit();
 	return true;
 }
 
+#if PLATFORM_WINDOWS
 int WINAPI WinMain(HINSTANCE hCurrInstance, HINSTANCE hPrevInstance, LPSTR lpCmdLine, int nShowCmd)
 {
 	hInstance = hCurrInstance;
 	return ZenDashboardMain(GetCommandLineW())? 0 : 1;
 }
 
+#elif PLATFORM_LINUX
+int main(int argc, char* argv[])
+{
+	return CommonUnixMain(argc, argv, &ZenDashboardMain);
+}
+#elif PLATFORM_MAC
+int main(int argc, char* argv[])
+{
+	[MacProgramDelegate mainWithArgc : argc argv : argv programMain : ZenDashboardMain programExit : FEngineLoop::AppExit] ;
+}
+#else
+#error "Unsupported platform!"
+#endif
+
 #undef LOCTEXT_NAMESPACE
diff --git a/Engine/Source/Programs/ZenDashboard/Private/ZenDashboardStyle.cpp b/Engine/Source/Programs/ZenDashboard/Private/ZenDashboardStyle.cpp
--- a/Engine/Source/Programs/ZenDashboard/Private/ZenDashboardStyle.cpp
+++ b/Engine/Source/Programs/ZenDashboard/Private/ZenDashboardStyle.cpp
@@ -29,10 +29,16 @@ void FZenDashboardStyle::Shutdown()
 }
 
 #define IMAGE_BRUSH( RelativePath, ... ) FSlateImageBrush(ContentFromEngine(TEXT(RelativePath), TEXT(".png")), __VA_ARGS__)
+#define IMAGE_BRUSH_SVG( RelativePath, ... ) FSlateVectorImageBrush(ContentFromEngine(TEXT(RelativePath), TEXT(".svg")), __VA_ARGS__)
 #define BOX_BRUSH( RelativePath, ... ) FSlateBoxBrush(ContentFromEngine(TEXT(RelativePath), TEXT(".png")), __VA_ARGS__)
 #define BORDER_BRUSH( RelativePath, ... ) FSlateBorderBrush(ContentFromEngine(RelativePath, TEXT(".png")), __VA_ARGS__)
 #define DEFAULT_FONT(...) FCoreStyle::GetDefaultFontStyle(__VA_ARGS__)
 
+#define EDITOR_IMAGE_BRUSH(RelativePath, ...) IMAGE_BRUSH("../Editor/Slate/" RelativePath, __VA_ARGS__)
+#define EDITOR_IMAGE_BRUSH_SVG(RelativePath, ...) IMAGE_BRUSH_SVG("../Editor/Slate/" RelativePath, __VA_ARGS__)
+#define EDITOR_BOX_BRUSH(RelativePath, ...) BOX_BRUSH("../Editor/Slate/" RelativePath, __VA_ARGS__)
+#define EDITOR_BORDER_BRUSH(RelativePath, ...) BORDER_BRUSH("../Editor/Slate/" RelativePath, __VA_ARGS__)
+
 namespace
 {
 	FString ContentFromEngine(const FString& RelativePath, const TCHAR* Extension)
@@ -49,9 +55,14 @@ TSharedRef< FSlateStyleSet > FZenDashboardStyle::Create()
 
 	Style.SetParentStyleName("CoreStyle");
 
+	const FVector2D Icon12x12(12.0f, 12.0f);
+	const FVector2D Icon16x16(16.0f, 16.0f);
 
 	Style.Set( "AppIcon", new IMAGE_BRUSH( "Icons/DefaultAppIcon", FVector2D(20,20)) );
 
+	Style.Set("Icons.FolderExplore", new EDITOR_IMAGE_BRUSH_SVG("Starship/Common/ContentBrowser", Icon12x12));
+	Style.Set("Icons.Clipboard", new EDITOR_IMAGE_BRUSH_SVG("Starship/Common/Clipboard", Icon16x16));
+
 	const FTextBlockStyle DefaultText = FTextBlockStyle()
 		.SetFont(DEFAULT_FONT("Bold", 10))
 		.SetColorAndOpacity(FSlateColor::UseForeground())
@@ -73,7 +84,6 @@ TSharedRef< FSlateStyleSet > FZenDashboardStyle::Create()
 	);
 
 
-	const FVector2D Icon16x16( 16.0f, 16.0f );
 	FSlateBrush* GenericWhiteBox = new IMAGE_BRUSH( "Old/White", Icon16x16 );
 
 	// Scrollbar
diff --git a/Engine/Source/Programs/ZenDashboard/Private/ZenDashboardStyle.h b/Engine/Source/Programs/ZenDashboard/Private/ZenDashboardStyle.h
--- a/Engine/Source/Programs/ZenDashboard/Private/ZenDashboardStyle.h
+++ b/Engine/Source/Programs/ZenDashboard/Private/ZenDashboardStyle.h
@@ -9,7 +9,7 @@ class FZenDashboardStyle
 {
 public:
 	/**
-	 * Set up specific styles for the crash report client app
+	 * Set up specific styles for the zen dashboard app
 	 */
 	static void Initialize();
 
diff --git a/Engine/Source/Programs/ZenDashboard/ZenDashboard.Build.cs b/Engine/Source/Programs/ZenDashboard/ZenDashboard.Build.cs
--- a/Engine/Source/Programs/ZenDashboard/ZenDashboard.Build.cs
+++ b/Engine/Source/Programs/ZenDashboard/ZenDashboard.Build.cs
@@ -10,17 +10,26 @@ public class ZenDashboard : ModuleRules
 			new string[]
 			{
 				"Core",
-				"CoreUObject",
 				"Json",
 				"Projects",
 				"ApplicationCore",
 				"Slate",
 				"SlateCore",
 				"StandaloneRenderer",
-				"Zen",
+				"Zen"
 				
 			});
 
+		if (Target.IsInPlatformGroup(UnrealPlatformGroup.Linux))
+		{
+			PrivateDependencyModuleNames.AddRange(
+				new string[]
+				{
+					"UnixCommonStartup"
+				}
+			);
+		}
+
 		PublicIncludePaths.Add("Runtime/Launch/Public");
 		PrivateIncludePaths.Add("Runtime/Launch/Private");      // For LaunchEngineLoop.cpp include
 	}
diff --git a/Engine/Source/Programs/ZenDashboard/ZenDashboard.Target.cs b/Engine/Source/Programs/ZenDashboard/ZenDashboard.Target.cs
--- a/Engine/Source/Programs/ZenDashboard/ZenDashboard.Target.cs
+++ b/Engine/Source/Programs/ZenDashboard/ZenDashboard.Target.cs
@@ -13,21 +13,33 @@ public class ZenDashboardTarget : TargetRules
 		IncludeOrderVersion = EngineIncludeOrderVersion.Latest;
 		LaunchModuleName = "ZenDashboard";
 
-		bBuildDeveloperTools = false;
-		bUseMallocProfiler = false;
+		bUseXGEController				= false;
+		bLoggingToMemoryEnabled			= true;
+		bUseLoggingInShipping			= true;
+		bCompileWithAccessibilitySupport= false;
+		bWithServerCode					= false;
+		bCompileNavmeshClusterLinks		= false;
+		bCompileNavmeshSegmentLinks		= false;
+		bCompileRecast					= false;
+		bCompileICU 					= false;
+		bWithLiveCoding					= false;
+		bBuildDeveloperTools			= false;
+		bUseMallocProfiler				= false;
+		bBuildWithEditorOnlyData		= false;
+		bCompileAgainstEngine			= false;
+		bCompileAgainstCoreUObject		= true;
+		bCompileAgainstApplicationCore	= true;
+		bUsesSlate						= true;
+		bForceDisableAutomationTests	= true;
 
-		// Editor-only data, however, is needed
-		bBuildWithEditorOnlyData = true;
-
-		// Currently this app is not linking against the engine, so we'll compile out references from Core to the rest of the engine
-		bCompileAgainstApplicationCore = true;
-		bCompileAgainstCoreUObject = true;
-		bCompileAgainstEngine = false;
+		bIsBuildingConsoleApplication = false;
 
-		// ICU is needed for regex during click to source code
-		bCompileICU = false;
+		WindowsPlatform.bUseBundledDbgHelp = false;
+		WindowsPlatform.bPixProfilingEnabled = false;
+		WindowsPlatform.bMergeIdenticalCOMDATs = true;
+		WindowsPlatform.bStripUnreferencedSymbols = true;
 
-		// UnrealHeaderTool is a console application, not a Windows app (sets entry point to main(), instead of WinMain())
-		bIsBuildingConsoleApplication = false;
+		GlobalDefinitions.Add("WITH_AUTOMATION_TESTS=0");
+		GlobalDefinitions.Add("WITH_AUTOMATION_WORKER=0");
 	}
 }
