-Fix non-unity compile issue
-Start ZenDashboard visible (even if it is operating as a system tray app) unless it is run with -Minimized on the commandline.

diff --git a/Engine/Source/Programs/ZenDashboard/Private/ZenDashboard.cpp b/Engine/Source/Programs/ZenDashboard/Private/ZenDashboard.cpp
--- a/Engine/Source/Programs/ZenDashboard/Private/ZenDashboard.cpp
+++ b/Engine/Source/Programs/ZenDashboard/Private/ZenDashboard.cpp
@@ -14,6 +14,7 @@
 #include "Modules/ModuleManager.h"
 #include "RequiredProgramMainCPPInclude.h"
 #include "ServiceInstanceManager.h"
+#include "SMessageDialog.h"
 #include "StandaloneRenderer.h"
 #include "SZenCacheStatistics.h"
 #include "SZenServiceStatus.h"
@@ -385,6 +386,7 @@ public:
 	void Run()
 	{
 		const bool bSystemTrayMode = !!PLATFORM_WINDOWS;
+		const bool bShowWindow = !(bSystemTrayMode && FParse::Param(FCommandLine::Get(), TEXT("Minimized")));
 		
 		Window =
 			SNew(SWindow)
@@ -437,19 +439,8 @@ public:
 			Window->SetRequestDestroyWindowOverride(FRequestDestroyWindowOverride::CreateStatic(&HideOnCloseOverride));
 		}
 
-		Slate.AddWindow(Window.ToSharedRef(), !bSystemTrayMode);
+		Slate.AddWindow(Window.ToSharedRef(), bShowWindow);
 
-		// Show the window without stealing focus
-/*		if (!FParse::Param(FCommandLine::Get(), TEXT("Hidden")))
-		{
-			HWND ForegroundWindow = GetForegroundWindow();
-			if (ForegroundWindow != nullptr)
-			{
-				::SetWindowPos((HWND)Window->GetNativeWindow()->GetOSWindowHandle(), ForegroundWindow, 0, 0, 0, 0, SWP_NOMOVE | SWP_NOSIZE);
-			}
-			Window->ShowWindow();
-		}
-		*/
 		// Setting focus seems to have to happen after the Window has been added
 		Slate.ClearKeyboardFocus(EFocusCause::Cleared);
 
