Add mechanism to run ZenServer garbage collection for fixed durations (one week, and one day) from the ZenDashboard.

diff --git a/Engine/Source/Developer/Zen/Private/ZenServerInterface.cpp b/Engine/Source/Developer/Zen/Private/ZenServerInterface.cpp
--- a/Engine/Source/Developer/Zen/Private/ZenServerInterface.cpp
+++ b/Engine/Source/Developer/Zen/Private/ZenServerInterface.cpp
@@ -2141,13 +2141,13 @@ FZenServiceInstance::RequestGC(const bool* OverrideCollectSmallObjects, const ui
 
 	if (OverrideCollectSmallObjects)
 	{
-		Query << Separators[SeparatorIndex] << LexToString(*OverrideCollectSmallObjects);
+		Query << Separators[SeparatorIndex] << TEXT("smallobjects=") << LexToString(*OverrideCollectSmallObjects);
 		SeparatorIndex = FMath::Min(SeparatorIndex + 1, (int32)UE_ARRAY_COUNT(Separators));
 	}
 
 	if (OverrideMaxCacheDuration)
 	{
-		Query << Separators[SeparatorIndex] << LexToString(*OverrideMaxCacheDuration);
+		Query << Separators[SeparatorIndex] << TEXT("maxcacheduration=") << LexToString(*OverrideMaxCacheDuration);
 		SeparatorIndex = FMath::Min(SeparatorIndex + 1, (int32)UE_ARRAY_COUNT(Separators));
 	}
 
diff --git a/Engine/Source/Programs/ZenDashboard/Private/ZenDashboard.cpp b/Engine/Source/Programs/ZenDashboard/Private/ZenDashboard.cpp
--- a/Engine/Source/Programs/ZenDashboard/Private/ZenDashboard.cpp
+++ b/Engine/Source/Programs/ZenDashboard/Private/ZenDashboard.cpp
@@ -11,6 +11,7 @@
 #include "Misc/CompilationResult.h"
 #include "Misc/MessageDialog.h"
 #include "Misc/Optional.h"
+#include "Misc/Timespan.h"
 #include "Modules/ModuleInterface.h"
 #include "Modules/ModuleManager.h"
 #include "RequiredProgramMainCPPInclude.h"
@@ -229,6 +230,24 @@ class FZenDashboardApp
 		}
 	}
 
+	void RunGCOneWeek()
+	{
+		if (TSharedPtr<UE::Zen::FZenServiceInstance> ZenServiceInstance = ServiceInstanceManager->GetZenServiceInstance())
+		{
+			uint32 CacheDuration = static_cast<uint32>(FTimespan::FromDays(7).GetTotalSeconds());
+			ZenServiceInstance->RequestGC(nullptr, &CacheDuration);
+		}
+	}
+
+	void RunGCOneDay()
+	{
+		if (TSharedPtr<UE::Zen::FZenServiceInstance> ZenServiceInstance = ServiceInstanceManager->GetZenServiceInstance())
+		{
+			uint32 CacheDuration = static_cast<uint32>(FTimespan::FromDays(1).GetTotalSeconds());
+			ZenServiceInstance->RequestGC(nullptr, &CacheDuration);
+		}
+	}
+
 	TSharedRef< SWidget > MakeMainMenu()
 	{
 		FMenuBarBuilder MenuBuilder( nullptr );
@@ -366,6 +385,30 @@ class FZenDashboardApp
 				NAME_None,
 				EUserInterfaceActionType::Button
 			);
+
+			MenuBuilder.AddMenuEntry(
+				LOCTEXT("FreeUnusedDiskSpaceOneWeek", "Free unused disk space (1 week age limit)"),
+				LOCTEXT("FreeUnusedDiskSpaceOneWeek_ToolTip", "Frees unused disk space by running a garbage collection process, only keeps data that was used in the past week"),
+				FSlateIcon(),
+				FUIAction(
+					FExecuteAction::CreateRaw(this, &FZenDashboardApp::RunGCOneWeek),
+					FCanExecuteAction()
+				),
+				NAME_None,
+				EUserInterfaceActionType::Button
+			);
+
+			MenuBuilder.AddMenuEntry(
+				LOCTEXT("FreeUnusedDiskSpaceOneDay", "Free unused disk space (1 day age limit)"),
+				LOCTEXT("FreeUnusedDiskSpaceOneDay_ToolTip", "Frees unused disk space by running a garbage collection process, only keeps data that was used in the past day"),
+				FSlateIcon(),
+				FUIAction(
+					FExecuteAction::CreateRaw(this, &FZenDashboardApp::RunGCOneDay),
+					FCanExecuteAction()
+				),
+				NAME_None,
+				EUserInterfaceActionType::Button
+			);
 		}
 		MenuBuilder.EndSection();
 	}
