Add more statistics to ZenDashboard

Cache: `requestcount`, `badrequestcount`, `writes`
Project Store: `requestcount`
Cid Store: `cidhits`, `cidmisses`, `cidwrites`

See companion PR for zenserver: https://github.com/EpicGames/zen/pull/405

diff --git a/Engine/Source/Developer/Zen/Private/ZenStatistics.cpp b/Engine/Source/Developer/Zen/Private/ZenStatistics.cpp
--- a/Engine/Source/Developer/Zen/Private/ZenStatistics.cpp
+++ b/Engine/Source/Developer/Zen/Private/ZenStatistics.cpp
@@ -37,9 +37,15 @@ bool LoadFromCompactBinary(FCbFieldView Field, FZenCacheStats::FGeneralStats& Ou
 	bOk = LoadFromCompactBinary(Field["size"], OutValue.Size) & bOk;
 	bOk = LoadFromCompactBinary(Field["hits"], OutValue.Hits) & bOk;
 	bOk = LoadFromCompactBinary(Field["misses"], OutValue.Misses) & bOk;
+	bOk = LoadFromCompactBinary(Field["writes"], OutValue.Writes) & bOk;
 	bOk = LoadFromCompactBinary(Field["hit_ratio"], OutValue.HitRatio) & bOk;
 	bOk = LoadFromCompactBinary(Field["upstream_hits"], OutValue.UpstreamHits) & bOk;
 	bOk = LoadFromCompactBinary(Field["upstream_ratio"], OutValue.UpstreamRatio) & bOk;
+	bOk = LoadFromCompactBinary(Field["cidhits"], OutValue.CidHits) & bOk;
+	bOk = LoadFromCompactBinary(Field["cidmisses"], OutValue.CidMisses) & bOk;
+	bOk = LoadFromCompactBinary(Field["cidwrites"], OutValue.CidWrites) & bOk;
+	bOk = LoadFromCompactBinary(Field["requestcount"], OutValue.RequestCount) & bOk;
+	bOk = LoadFromCompactBinary(Field["badrequestcount"], OutValue.BadRequestCount) & bOk;
 	return bOk;
 }
 
@@ -139,6 +145,7 @@ bool LoadFromCompactBinary(FCbFieldView Field, FZenProjectStats::FGeneralStats&
 	bOk = LoadFromCompactBinary(Field["oplog"], OutValue.Oplog) & bOk;
 	bOk = LoadFromCompactBinary(Field["op"], OutValue.Op) & bOk;
 	bOk = LoadFromCompactBinary(Field["chunk"], OutValue.Chunk) & bOk;
+	bOk = LoadFromCompactBinary(Field["requestcount"], OutValue.RequestCount) & bOk;
 	bOk = LoadFromCompactBinary(Field["badrequestcount"], OutValue.BadRequestCount) & bOk;
 	return bOk;
 }
diff --git a/Engine/Source/Developer/Zen/Public/Experimental/ZenStatistics.h b/Engine/Source/Developer/Zen/Public/Experimental/ZenStatistics.h
--- a/Engine/Source/Developer/Zen/Public/Experimental/ZenStatistics.h
+++ b/Engine/Source/Developer/Zen/Public/Experimental/ZenStatistics.h
@@ -40,9 +40,15 @@ struct FZenCacheStats
 		FZenSizeStats Size;
 		int64 Hits = 0;
 		int64 Misses = 0;
+		int64 Writes = 0;
 		double HitRatio = 0.0;
 		int64 UpstreamHits = 0;
 		double UpstreamRatio = 0.0;
+		int64 CidHits = 0;
+		int64 CidMisses = 0;
+		int64 CidWrites = 0;
+		int64 RequestCount = 0;
+		int64 BadRequestCount = 0;
 	};
 
 	struct FRequestStats
@@ -105,6 +111,7 @@ struct FZenProjectStats
 		FReadWriteDeleteStats Oplog;
 		FHitMissWriteStats Op;
 		FHitMissWriteStats Chunk;
+		int64 RequestCount;
 		int64 BadRequestCount;
 	};
 	FGeneralStats General;
diff --git a/Engine/Source/Programs/ZenDashboard/Private/SZenCacheStatistics.cpp b/Engine/Source/Programs/ZenDashboard/Private/SZenCacheStatistics.cpp
--- a/Engine/Source/Programs/ZenDashboard/Private/SZenCacheStatistics.cpp
+++ b/Engine/Source/Programs/ZenDashboard/Private/SZenCacheStatistics.cpp
@@ -126,21 +126,86 @@ TSharedRef<SWidget> SZenCacheStatistics::GetGridPanel()
 	Row++;
 
 	Panel->AddSlot(0, Row)
-		[
-			SNew(STextBlock)
-			.Margin(FMargin(ColumnMargin, RowMargin))
-			.Text(LOCTEXT("CacheMissQuantity", "Misses"))
-		];
+	[
+		SNew(STextBlock)
+		.Margin(FMargin(ColumnMargin, RowMargin))
+		.Text(LOCTEXT("CacheMissQuantity", "Misses"))
+	];
 
 	Panel->AddSlot(1, Row)
-		[
-			SNew(STextBlock)
-			.Margin(FMargin(ColumnMargin, RowMargin))
-			.Text_Lambda([bHaveStats, CacheMisses = ZenStats.General.Misses]
+	[
+		SNew(STextBlock)
+		.Margin(FMargin(ColumnMargin, RowMargin))
+		.Text_Lambda([bHaveStats, CacheMisses = ZenStats.General.Misses]
+		{
+			return bHaveStats ? FText::AsNumber(CacheMisses) : LOCTEXT("UnavailableValue", "-");
+		})
+	];
+
+	Row++;
+
+	Panel->AddSlot(0, Row)
+	[
+		SNew(STextBlock)
+		.Margin(FMargin(ColumnMargin, RowMargin))
+		.Text(LOCTEXT("CacheWriteQuantity", "Writes"))
+	];
+
+	Panel->AddSlot(1, Row)
+	[
+		SNew(STextBlock)
+		.Margin(FMargin(ColumnMargin, RowMargin))
+		.Text_Lambda([bHaveStats, CacheWrites = ZenStats.General.Writes]
+		{
+			return bHaveStats ? FText::AsNumber(CacheWrites) : LOCTEXT("UnavailableValue", "-");
+		})
+	];
+
+	Row++;
+
+	Panel->AddSlot(0, Row)
+	[
+		SNew(STextBlock)
+		.Margin(FMargin(ColumnMargin, RowMargin))
+		.Text(LOCTEXT("CacheRequests", "Requests"))
+	];
+
+	Panel->AddSlot(1, Row)
+	[
+		SNew(STextBlock)
+		.Margin(FMargin(ColumnMargin, RowMargin))
+		.Text_Lambda([bHaveStats, Requests = ZenStats.General.RequestCount]
+		{
+			if (bHaveStats)
 			{
-				return bHaveStats ? FText::AsNumber(CacheMisses) : LOCTEXT("UnavailableValue", "-");
-			})
-		];
+				return FText::AsNumber(Requests);
+			}
+			return LOCTEXT("UnavailableValue", "-");
+		})
+	];
+
+	Row++;
+
+	Panel->AddSlot(0, Row)
+	[
+		SNew(STextBlock)
+		.Margin(FMargin(ColumnMargin, RowMargin))
+		.Text(LOCTEXT("CacheBadRequests", "Bad Requests"))
+	];
+
+	Panel->AddSlot(1, Row)
+	[
+		SNew(STextBlock)
+		.Margin(FMargin(ColumnMargin, RowMargin))
+		.Text_Lambda([bHaveStats, BadRequests = ZenStats.General.BadRequestCount]
+		{
+			if (bHaveStats)
+			{
+				return FText::AsNumber(BadRequests);
+			}
+			return LOCTEXT("UnavailableValue", "-");
+		})
+	];
 
 	Row++;
 
diff --git a/Engine/Source/Programs/ZenDashboard/Private/SZenCidStoreStatistics.cpp b/Engine/Source/Programs/ZenDashboard/Private/SZenCidStoreStatistics.cpp
new file mode 100644
--- /dev/null
+++ b/Engine/Source/Programs/ZenDashboard/Private/SZenCidStoreStatistics.cpp
@@ -0,0 +1,213 @@
+// Copyright Epic Games, Inc. All Rights Reserved.
+
+#include "SZenCidStoreStatistics.h"
+#include "Experimental/ZenServerInterface.h"
+#include "Internationalization/FastDecimalFormat.h"
+#include "Math/BasicMathExpressionEvaluator.h"
+#include "Math/UnitConversion.h"
+#include "Misc/ExpressionParser.h"
+#include "Styling/StyleColors.h"
+#include "Widgets/Layout/SGridPanel.h"
+#include "Widgets/SBoxPanel.h"
+#include "Widgets/Text/STextBlock.h"
+
+#define LOCTEXT_NAMESPACE "ZenDashboard"
+
+void SZenCidStoreStatistics::Construct(const FArguments& InArgs)
+{
+	ZenServiceInstance = InArgs._ZenServiceInstance;
+
+	this->ChildSlot
+	[
+		SNew(SVerticalBox)
+
+		+ SVerticalBox::Slot()
+		.AutoHeight()
+		.Padding(0, 0, 0, 0)
+		.Expose(GridSlot)
+		[
+			GetGridPanel()
+		]
+	];
+
+	RegisterActiveTimer(0.5f, FWidgetActiveTimerDelegate::CreateSP(this, &SZenCidStoreStatistics::UpdateGridPanels));
+}
+
+EActiveTimerReturnType SZenCidStoreStatistics::UpdateGridPanels(double InCurrentTime, float InDeltaTime)
+{
+	(*GridSlot)
+	[
+		GetGridPanel()
+	];
+
+	SlatePrepass(GetPrepassLayoutScaleMultiplier());
+
+	return EActiveTimerReturnType::Continue;
+}
+
+TSharedRef<SWidget> SZenCidStoreStatistics::GetGridPanel()
+{
+	TSharedRef<SGridPanel> Panel = SNew(SGridPanel);
+
+	bool bHaveStats = false;
+	UE::Zen::FZenCacheStats ZenStats;
+
+	if (TSharedPtr<UE::Zen::FZenServiceInstance> ServiceInstance = ZenServiceInstance.Get())
+	{
+		if (ServiceInstance->GetCacheStats(ZenStats))
+		{
+			bHaveStats = true;
+		}
+	}
+
+	int64 TotalHits = ZenStats.General.CidHits;
+	int64 TotalMisses = ZenStats.General.CidMisses;
+	int64 TotalWrites = ZenStats.General.CidWrites;
+	int64 Requests = ZenStats.General.CidHits + ZenStats.General.CidMisses + ZenStats.General.CidWrites;
+
+	const static FNumberFormattingOptions SingleDecimalFormatting = FNumberFormattingOptions()
+		.SetUseGrouping(true)
+		.SetMinimumFractionalDigits(1)
+		.SetMaximumFractionalDigits(1);
+
+	int32 Row = 0;
+
+	const float RowMargin = 0.0f;
+	const float TitleMargin = 10.0f;
+	const float ColumnMargin = 10.0f;
+	const FSlateColor TitleColor = FStyleColors::AccentWhite;
+	const FSlateFontInfo TitleFont = FCoreStyle::GetDefaultFontStyle("Bold", 10);
+
+
+	Panel->AddSlot(0, Row)
+	.ColumnSpan(2)
+	[
+		SNew(STextBlock)
+		.Margin(FMargin(ColumnMargin, RowMargin))
+		.ColorAndOpacity(TitleColor)
+		.Font(TitleFont)
+		.Text(LOCTEXT("CAS", "Local Cid Store"))
+	];
+
+	Row++;
+
+	Panel->AddSlot(0, Row)
+	[
+		SNew(STextBlock)
+		.Margin(FMargin(ColumnMargin, RowMargin))
+		.Text(LOCTEXT("CidStoreHitRate", "Hit Rate"))
+	];
+
+	Panel->AddSlot(1, Row)
+	[
+		SNew(STextBlock)
+		.Margin(FMargin(ColumnMargin, RowMargin))
+		.Text_Lambda([bHaveStats, TotalHits, TotalMisses]
+		{
+			if (bHaveStats && ((TotalHits+TotalMisses) > 0))
+			{
+				double HitRate = static_cast<double>(TotalHits) / static_cast<double>(TotalHits + TotalMisses);
+				return FText::AsPercent(HitRate, &SingleDecimalFormatting);
+			}
+			return LOCTEXT("UnavailableValue", "-");
+		})
+	];
+
+	Row++;
+
+	Panel->AddSlot(0, Row)
+	[
+		SNew(STextBlock)
+		.Margin(FMargin(ColumnMargin, RowMargin))
+		.Text(LOCTEXT("CidStoreHitQuantity", "Hits"))
+	];
+
+	Panel->AddSlot(1, Row)
+	[
+		SNew(STextBlock)
+		.Margin(FMargin(ColumnMargin, RowMargin))
+		.Text_Lambda([bHaveStats, TotalHits]
+		{
+			if (bHaveStats)
+			{
+				return FText::AsNumber(TotalHits);
+			}
+			return LOCTEXT("UnavailableValue", "-");
+		})
+	];
+
+	Row++;
+
+	Panel->AddSlot(0, Row)
+	[
+		SNew(STextBlock)
+		.Margin(FMargin(ColumnMargin, RowMargin))
+		.Text(LOCTEXT("CidStoreMissQuantity", "Misses"))
+	];
+
+	Panel->AddSlot(1, Row)
+	[
+		SNew(STextBlock)
+		.Margin(FMargin(ColumnMargin, RowMargin))
+		.Text_Lambda([bHaveStats, TotalMisses]
+		{
+			if (bHaveStats)
+			{
+				return FText::AsNumber(TotalMisses);
+			}
+			return LOCTEXT("UnavailableValue", "-");
+		})
+	];
+
+	Row++;
+
+	Panel->AddSlot(0, Row)
+	[
+		SNew(STextBlock)
+		.Margin(FMargin(ColumnMargin, RowMargin))
+		.Text(LOCTEXT("CidStoreWriteQuantity", "Writes"))
+	];
+
+	Panel->AddSlot(1, Row)
+	[
+		SNew(STextBlock)
+		.Margin(FMargin(ColumnMargin, RowMargin))
+		.Text_Lambda([bHaveStats, TotalWrites]
+		{
+			if (bHaveStats)
+			{
+				return FText::AsNumber(TotalWrites);
+			}
+			return LOCTEXT("UnavailableValue", "-");
+		})
+	];
+
+	Row++;
+
+	Panel->AddSlot(0, Row)
+	[
+		SNew(STextBlock)
+		.Margin(FMargin(ColumnMargin, RowMargin))
+		.Text(LOCTEXT("ProjectRequests", "Requests"))
+	];
+
+	Panel->AddSlot(1, Row)
+	[
+		SNew(STextBlock)
+		.Margin(FMargin(ColumnMargin, RowMargin))
+		.Text_Lambda([bHaveStats, Requests]
+		{
+			if (bHaveStats)
+			{
+				return FText::AsNumber(Requests);
+			}
+			return LOCTEXT("UnavailableValue", "-");
+		})
+	];
+
+	Row++;
+
+	return Panel;
+}
+
+#undef LOCTEXT_NAMESPACE
diff --git a/Engine/Source/Programs/ZenDashboard/Private/SZenCidStoreStatistics.h b/Engine/Source/Programs/ZenDashboard/Private/SZenCidStoreStatistics.h
new file mode 100644
--- /dev/null
+++ b/Engine/Source/Programs/ZenDashboard/Private/SZenCidStoreStatistics.h
@@ -0,0 +1,33 @@
+// Copyright Epic Games, Inc. All Rights Reserved.
+
+#pragma once
+
+#include "Experimental/ZenServerInterface.h"
+#include "ServiceInstanceManager.h"
+#include "Templates/SharedPointer.h"
+#include "Widgets/DeclarativeSyntaxSupport.h"
+#include "Widgets/SBoxPanel.h"
+#include "Widgets/SCompoundWidget.h"
+#include "Widgets/SWidget.h"
+
+class SZenCidStoreStatistics : public SCompoundWidget
+{
+	SLATE_BEGIN_ARGS(SZenCidStoreStatistics)
+		: _ZenServiceInstance(nullptr)
+	{ }
+
+	SLATE_ATTRIBUTE(TSharedPtr<UE::Zen::FZenServiceInstance>, ZenServiceInstance);
+
+	SLATE_END_ARGS()
+
+	void Construct(const FArguments& InArgs);
+
+private:
+
+	TSharedRef<SWidget> GetGridPanel();
+
+	EActiveTimerReturnType UpdateGridPanels(double InCurrentTime, float InDeltaTime);
+
+	SVerticalBox::FSlot* GridSlot = nullptr;
+	TAttribute<TSharedPtr<UE::Zen::FZenServiceInstance>> ZenServiceInstance;
+};
\ No newline at end of file
diff --git a/Engine/Source/Programs/ZenDashboard/Private/SZenProjectStatistics.cpp b/Engine/Source/Programs/ZenDashboard/Private/SZenProjectStatistics.cpp
--- a/Engine/Source/Programs/ZenDashboard/Private/SZenProjectStatistics.cpp
+++ b/Engine/Source/Programs/ZenDashboard/Private/SZenProjectStatistics.cpp
@@ -187,18 +187,18 @@ TSharedRef<SWidget> SZenProjectStatistics::GetGridPanel()
 	[
 		SNew(STextBlock)
 		.Margin(FMargin(ColumnMargin, RowMargin))
-		.Text(LOCTEXT("ProjectBadRequests", "Bad Requests"))
+		.Text(LOCTEXT("ProjectRequests", "Requests"))
 	];
 
 	Panel->AddSlot(1, Row)
 	[
 		SNew(STextBlock)
 		.Margin(FMargin(ColumnMargin, RowMargin))
-		.Text_Lambda([bHaveStats, BadRequests = ZenStats.General.BadRequestCount]
+		.Text_Lambda([bHaveStats, Requests = ZenStats.General.RequestCount]
 		{
 			if (bHaveStats)
 			{
-				return FText::AsNumber(BadRequests);
+				return FText::AsNumber(Requests);
 			}
 			return LOCTEXT("UnavailableValue", "-");
 		})
@@ -206,6 +206,28 @@ TSharedRef<SWidget> SZenProjectStatistics::GetGridPanel()
 
 	Row++;
 
+	Panel->AddSlot(0, Row)
+		[
+			SNew(STextBlock)
+			.Margin(FMargin(ColumnMargin, RowMargin))
+		.Text(LOCTEXT("ProjectBadRequests", "Bad Requests"))
+		];
+
+	Panel->AddSlot(1, Row)
+		[
+			SNew(STextBlock)
+			.Margin(FMargin(ColumnMargin, RowMargin))
+		.Text_Lambda([bHaveStats, BadRequests = ZenStats.General.BadRequestCount]
+			{
+				if (bHaveStats)
+				{
+					return FText::AsNumber(BadRequests);
+				}
+	return LOCTEXT("UnavailableValue", "-");
+			})
+		];
+
+	Row++;
 
 	return Panel;
 }
diff --git a/Engine/Source/Programs/ZenDashboard/Private/ZenDashboard.cpp b/Engine/Source/Programs/ZenDashboard/Private/ZenDashboard.cpp
--- a/Engine/Source/Programs/ZenDashboard/Private/ZenDashboard.cpp
+++ b/Engine/Source/Programs/ZenDashboard/Private/ZenDashboard.cpp
@@ -19,6 +19,7 @@
 #include "SMessageDialog.h"
 #include "StandaloneRenderer.h"
 #include "SZenCacheStatistics.h"
+#include "SZenCidStoreStatistics.h"
 #include "SZenProjectStatistics.h"
 #include "SZenServiceStatus.h"
 #include "Templates/SharedPointer.h"
@@ -489,6 +490,15 @@ public:
 							SNew(SZenProjectStatistics)
 							.ZenServiceInstance(ServiceInstanceManager.ToSharedRef(), &UE::Zen::FServiceInstanceManager::GetZenServiceInstance)
 						]
+
+						+SHorizontalBox::Slot()
+						.AutoWidth()
+						.HAlign(HAlign_Left)
+						.VAlign(VAlign_Top)
+						[
+							SNew(SZenCidStoreStatistics)
+							.ZenServiceInstance(ServiceInstanceManager.ToSharedRef(), &UE::Zen::FServiceInstanceManager::GetZenServiceInstance)
+						]
 					]
 				]
 			];
