Allow ZenDashboard to be built for non-windows desktop platforms.

diff --git a/Engine/Source/Programs/ZenDashboard/ZenDashboard.Target.cs b/Engine/Source/Programs/ZenDashboard/ZenDashboard.Target.cs
--- a/Engine/Source/Programs/ZenDashboard/ZenDashboard.Target.cs
+++ b/Engine/Source/Programs/ZenDashboard/ZenDashboard.Target.cs
@@ -3,7 +3,7 @@
 using UnrealBuildTool;
 using System.Collections.Generic;
 
-[SupportedPlatforms("Win64")]
+[SupportedPlatforms(UnrealPlatformClass.Desktop)]
 public class ZenDashboardTarget : TargetRules
 {
 	public ZenDashboardTarget(TargetInfo Target) : base(Target)
