Change Zen data delete operation to be asynchronous and ensure that other exclusive operations can't be instigated while an asynchronous operation is in progress.

diff --git a/Engine/Source/Programs/ZenDashboard/Private/ZenDashboard.cpp b/Engine/Source/Programs/ZenDashboard/Private/ZenDashboard.cpp
--- a/Engine/Source/Programs/ZenDashboard/Private/ZenDashboard.cpp
+++ b/Engine/Source/Programs/ZenDashboard/Private/ZenDashboard.cpp
@@ -22,6 +22,7 @@
 #include "SZenCidStoreStatistics.h"
 #include "SZenProjectStatistics.h"
 #include "SZenServiceStatus.h"
+#include "Tasks/Task.h"
 #include "Templates/SharedPointer.h"
 #include "Widgets/Input/SButton.h"
 #include "Widgets/Input/SButton.h"
@@ -82,6 +83,8 @@ class FZenDashboardApp
 	TSharedPtr<UE::Zen::FServiceInstanceManager> ServiceInstanceManager;
 	TArray<FSimpleDelegate> MainThreadTasks;
 
+	std::atomic<bool> bLatentExclusiveOperationActive = false;
+
 #if PLATFORM_WINDOWS
 	bool ProcessMessage(HWND hwnd, uint32 msg, WPARAM wParam, LPARAM lParam, int32& OutResult) override
 	{
@@ -133,16 +136,28 @@ class FZenDashboardApp
 				Window->ShowWindow();
 				break;
 			case CMD_STARTZENSERVER:
-				StartZenServer();
+				if (CanExecuteExclusiveAction())
+				{
+					StartZenServer();
+				}
 				break;
 			case CMD_STOPZENSERVER:
-				StopZenServer();
+				if (CanExecuteExclusiveAction())
+				{
+					StopZenServer();
+				}
 				break;
 			case CMD_RESTARTZENSERVER:
-				RestartZenServer();
+				if (CanExecuteExclusiveAction())
+				{
+					RestartZenServer();
+				}
 				break;
 			case CMD_EXITDASHBOARD:
-				ExitDashboard();
+				if (CanExecuteExclusiveAction())
+				{
+					ExitDashboard();
+				}
 				break;
 			}
 			break;
@@ -157,6 +172,11 @@ class FZenDashboardApp
 		FSlateApplication::Get().RequestDestroyWindow(Window.ToSharedRef());
 	}
 
+	bool CanExecuteExclusiveAction()
+	{
+		return !bLatentExclusiveOperationActive;
+	}
+
 	void StartZenServer()
 	{
 		UE::Zen::FZenLocalServiceRunContext RunContext;
@@ -199,9 +219,20 @@ class FZenDashboardApp
 
 				StopZenServer();
 
-				FPlatformFileManager::Get().GetPlatformFile().DeleteDirectoryRecursively(*RunContext.GetDataPath());
+				bLatentExclusiveOperationActive = true;
+
+				UE::Tasks::Launch(TEXT("DeleteZenDataAndRestart"),
+					[this, DataPath = RunContext.GetDataPath()] ()
+					{
+						FPlatformProcess::Sleep(10.0f);
+						//FPlatformFileManager::Get().GetPlatformFile().DeleteDirectoryRecursively(DataPath);
+
+						StartZenServer();
 
-				StartZenServer();
+						bLatentExclusiveOperationActive = false;
+					},
+					UE::Tasks::ETaskPriority::BackgroundNormal
+				);
 			}
 		}
 	}
@@ -280,7 +311,7 @@ class FZenDashboardApp
 			FSlateIcon(),
 			FUIAction(
 				FExecuteAction::CreateRaw( this, &FZenDashboardApp::ExitDashboard ),
-				FCanExecuteAction()
+				FCanExecuteAction::CreateRaw(this, &FZenDashboardApp::CanExecuteExclusiveAction)
 			),
 			NAME_None,
 			EUserInterfaceActionType::Button
@@ -297,7 +328,7 @@ class FZenDashboardApp
 				FSlateIcon(),
 				FUIAction(
 					FExecuteAction::CreateRaw(this, &FZenDashboardApp::StartZenServer),
-					FCanExecuteAction()
+					FCanExecuteAction::CreateRaw(this, &FZenDashboardApp::CanExecuteExclusiveAction)
 				),
 				NAME_None,
 				EUserInterfaceActionType::Button
@@ -308,7 +339,7 @@ class FZenDashboardApp
 				FSlateIcon(),
 				FUIAction(
 					FExecuteAction::CreateRaw(this, &FZenDashboardApp::StopZenServer),
-					FCanExecuteAction()
+					FCanExecuteAction::CreateRaw(this, &FZenDashboardApp::CanExecuteExclusiveAction)
 				),
 				NAME_None,
 				EUserInterfaceActionType::Button
@@ -319,7 +350,7 @@ class FZenDashboardApp
 				FSlateIcon(),
 				FUIAction(
 					FExecuteAction::CreateRaw(this, &FZenDashboardApp::RestartZenServer),
-					FCanExecuteAction()
+					FCanExecuteAction::CreateRaw(this, &FZenDashboardApp::CanExecuteExclusiveAction)
 				),
 				NAME_None,
 				EUserInterfaceActionType::Button
@@ -334,7 +365,7 @@ class FZenDashboardApp
 							FSlateIcon(),
 							FUIAction(
 								FExecuteAction::CreateRaw(this, &FZenDashboardApp::DeleteDataAndRestartZenServer),
-								FCanExecuteAction()
+								FCanExecuteAction::CreateRaw(this, &FZenDashboardApp::CanExecuteExclusiveAction)
 							),
 							NAME_None,
 							EUserInterfaceActionType::Button
@@ -380,7 +411,7 @@ class FZenDashboardApp
 				FSlateIcon(),
 				FUIAction(
 					FExecuteAction::CreateRaw(this, &FZenDashboardApp::RunGC),
-					FCanExecuteAction()
+					FCanExecuteAction::CreateRaw(this, &FZenDashboardApp::CanExecuteExclusiveAction)
 				),
 				NAME_None,
 				EUserInterfaceActionType::Button
@@ -392,7 +423,7 @@ class FZenDashboardApp
 				FSlateIcon(),
 				FUIAction(
 					FExecuteAction::CreateRaw(this, &FZenDashboardApp::RunGCOneWeek),
-					FCanExecuteAction()
+					FCanExecuteAction::CreateRaw(this, &FZenDashboardApp::CanExecuteExclusiveAction)
 				),
 				NAME_None,
 				EUserInterfaceActionType::Button
@@ -404,7 +435,7 @@ class FZenDashboardApp
 				FSlateIcon(),
 				FUIAction(
 					FExecuteAction::CreateRaw(this, &FZenDashboardApp::RunGCOneDay),
-					FCanExecuteAction()
+					FCanExecuteAction::CreateRaw(this, &FZenDashboardApp::CanExecuteExclusiveAction)
 				),
 				NAME_None,
 				EUserInterfaceActionType::Button
@@ -653,7 +684,7 @@ int ZenDashboardMain(const TCHAR* CmdLine)
 }
 
 #if PLATFORM_WINDOWS
-int WINAPI WinMain(HINSTANCE hCurrInstance, HINSTANCE hPrevInstance, LPSTR lpCmdLine, int nShowCmd)
+int WINAPI WinMain(_In_ HINSTANCE hCurrInstance, _In_opt_ HINSTANCE hPrevInstance, _In_ LPSTR lpCmdLine, _In_ int nShowCmd)
 {
 	hInstance = hCurrInstance;
 	return ZenDashboardMain(GetCommandLineW())? 0 : 1;
