Restore deletion of directory from ZenDashboard that was accidentally commented out.  Also changed main window title to Unreal Zen Dashboard.

diff --git a/Engine/Source/Programs/ZenDashboard/Private/ZenDashboard.cpp b/Engine/Source/Programs/ZenDashboard/Private/ZenDashboard.cpp
--- a/Engine/Source/Programs/ZenDashboard/Private/ZenDashboard.cpp
+++ b/Engine/Source/Programs/ZenDashboard/Private/ZenDashboard.cpp
@@ -225,7 +225,7 @@ class FZenDashboardApp
 					[this, DataPath = RunContext.GetDataPath()] ()
 					{
 						FPlatformProcess::Sleep(10.0f);
-						//FPlatformFileManager::Get().GetPlatformFile().DeleteDirectoryRecursively(DataPath);
+						FPlatformFileManager::Get().GetPlatformFile().DeleteDirectoryRecursively(*DataPath);
 
 						StartZenServer();
 
@@ -593,7 +593,7 @@ public:
 private:
 	FText GetWindowTitle()
 	{
-		return LOCTEXT("WindowTitle", "Zen Dashboard");
+		return LOCTEXT("WindowTitle", "Unreal Zen Dashboard");
 	}
 
 	EAppReturnType::Type OnModalMessageDialog(EAppMsgCategory InMessageCategory, EAppMsgType::Type InMessage, const FText& InText, const FText& InTitle)
