Make ZenDashboard behave as a standard (non systray) app on Windows.  Remove the search widgets from the ZenDashboard menus. Add an entry to the DerivedData status bar to launch (and if needed build) ZenDashboard.

diff --git a/Engine/Source/Editor/DerivedDataEditor/DerivedDataEditor.Build.cs b/Engine/Source/Editor/DerivedDataEditor/DerivedDataEditor.Build.cs
--- a/Engine/Source/Editor/DerivedDataEditor/DerivedDataEditor.Build.cs
+++ b/Engine/Source/Editor/DerivedDataEditor/DerivedDataEditor.Build.cs
@@ -24,6 +24,7 @@ public class DerivedDataEditor : ModuleRules
 				"WorkspaceMenuStructure",
 				"MessageLog",
 				"ToolWidgets",
+				"UATHelper",
 				"Zen"
 			});
 
diff --git a/Engine/Source/Editor/DerivedDataEditor/Private/SDerivedDataStatusBar.cpp b/Engine/Source/Editor/DerivedDataEditor/Private/SDerivedDataStatusBar.cpp
--- a/Engine/Source/Editor/DerivedDataEditor/Private/SDerivedDataStatusBar.cpp
+++ b/Engine/Source/Editor/DerivedDataEditor/Private/SDerivedDataStatusBar.cpp
@@ -45,6 +45,7 @@
 #include "Widgets/SBoxPanel.h"
 #include "Widgets/SOverlay.h"
 #include "Widgets/Text/STextBlock.h"
+#include "ZenDashboardLauncher.h"
 
 class SWidget;
 struct FSlateBrush;
@@ -69,6 +70,7 @@ void FDerivedDataStatusBarMenuCommands::RegisterCommands()
 	UI_COMMAND(ChangeSettings, "Change Cache Settings", "Opens a dialog to change Cache settings.", EUserInterfaceActionType::Button, FInputChord());
 	UI_COMMAND(ViewCacheStatistics, "View Cache Statistics", "Opens the Cache Statistics panel.", EUserInterfaceActionType::Button, FInputChord());
 	UI_COMMAND(ViewResourceUsage, "View Resource Usage", "Opens the Resource Usage panel.", EUserInterfaceActionType::Button, FInputChord());
+	UI_COMMAND(LaunchZenDashboard, "Launch Zen Dashboard", "Launches the Zen Dashboard utility.", EUserInterfaceActionType::Button, FInputChord());
 
 	ActionList->MapAction(
 		ChangeSettings,
@@ -84,6 +86,11 @@ void FDerivedDataStatusBarMenuCommands::RegisterCommands()
 		ViewResourceUsage,
 		FExecuteAction::CreateStatic(&FDerivedDataStatusBarMenuCommands::ViewResourceUsage_Clicked)
 	);
+
+	ActionList->MapAction(
+		LaunchZenDashboard,
+		FExecuteAction::CreateStatic(&FDerivedDataStatusBarMenuCommands::LaunchZenDashboard_Clicked)
+	);
 }
 
 void FDerivedDataStatusBarMenuCommands::ChangeSettings_Clicked()
@@ -103,6 +110,12 @@ void FDerivedDataStatusBarMenuCommands::ViewResourceUsage_Clicked()
 	DerivedDataEditorModule.ShowResourceUsageTab();
 }
 
+void FDerivedDataStatusBarMenuCommands::LaunchZenDashboard_Clicked()
+{
+	using namespace UE::Zen;
+	FZenDashboardLauncher::Get()->StartZenDashboard(FZenDashboardLauncher::Get()->GetZenDashboardApplicationPath());
+}
+
 TSharedRef<SWidget> SDerivedDataStatusBarWidget::CreateStatusBarMenu()
 {
 	UToolMenu* Menu = UToolMenus::Get()->RegisterMenu("StatusBar.ToolBar.DDC", NAME_None, EMultiBoxType::Menu, false);
@@ -116,6 +129,12 @@ TSharedRef<SWidget> SDerivedDataStatusBarWidget::CreateStatusBarMenu()
 			TAttribute<FText>(),
 			FSlateIcon(FAppStyle::GetAppStyleSetName(), "DerivedData.Cache.Settings")
 		);
+		Section.AddMenuEntry(
+			FDerivedDataStatusBarMenuCommands::Get().LaunchZenDashboard,
+			TAttribute<FText>(),
+			TAttribute<FText>(),
+			FSlateIcon(FAppStyle::GetAppStyleSetName(), "DerivedData.LaunchZenDashboard")
+		);
 	}
 
 	{
diff --git a/Engine/Source/Editor/DerivedDataEditor/Private/SDerivedDataStatusBar.h b/Engine/Source/Editor/DerivedDataEditor/Private/SDerivedDataStatusBar.h
--- a/Engine/Source/Editor/DerivedDataEditor/Private/SDerivedDataStatusBar.h
+++ b/Engine/Source/Editor/DerivedDataEditor/Private/SDerivedDataStatusBar.h
@@ -28,12 +28,14 @@ private:
 	static void ChangeSettings_Clicked();
 	static void ViewCacheStatistics_Clicked();
 	static void ViewResourceUsage_Clicked();
+	static void LaunchZenDashboard_Clicked();
 
 public:
 
 	TSharedPtr< FUICommandInfo > ChangeSettings;
 	TSharedPtr< FUICommandInfo > ViewResourceUsage;
 	TSharedPtr< FUICommandInfo > ViewCacheStatistics;
+	TSharedPtr< FUICommandInfo > LaunchZenDashboard;
 
 	static TSharedRef<FUICommandList> ActionList;
 };
diff --git a/Engine/Source/Editor/DerivedDataEditor/Private/ZenDashboardLauncher.cpp b/Engine/Source/Editor/DerivedDataEditor/Private/ZenDashboardLauncher.cpp
new file mode 100644
--- /dev/null
+++ b/Engine/Source/Editor/DerivedDataEditor/Private/ZenDashboardLauncher.cpp
@@ -0,0 +1,129 @@
+// Copyright Epic Games, Inc. All Rights Reserved.
+
+#include "ZenDashboardLauncher.h"
+
+#if UE_WITH_ZEN
+
+#include "IUATHelperModule.h"
+#include "Logging/LogMacros.h"
+#include "Logging/MessageLog.h"
+#include "Misc/Paths.h"
+#include "Styling/AppStyle.h"
+
+#define LOCTEXT_NAMESPACE "FZenDashboardLauncher"
+
+namespace UE::Zen
+{
+
+DEFINE_LOG_CATEGORY_STATIC(LogZenDashboardLauncher, Log, All);
+
+TSharedPtr<FZenDashboardLauncher> FZenDashboardLauncher::Instance = nullptr;
+
+FZenDashboardLauncher::FZenDashboardLauncher()
+{
+
+}
+
+FZenDashboardLauncher::~FZenDashboardLauncher()
+{
+
+}
+	
+FString FZenDashboardLauncher::GetZenDashboardApplicationPath()
+{
+	FString Path = FPlatformProcess::GenerateApplicationPath(TEXT("ZenDashboard"), EBuildConfiguration::Development);
+	return FPaths::ConvertRelativePathToFull(Path);
+}
+
+void FZenDashboardLauncher::StartZenDashboard(const FString& Path, const FString& Parameters)
+{
+	auto Callback = [](const EStartZenDashboardResult Result) {};
+	StartZenDashboard(Path, Parameters, Callback);
+}
+
+void FZenDashboardLauncher::StartZenDashboard(const FString& Path, const FString& Parameters, StartZenDashboardCallback Callback)
+{
+    if (!FPaths::FileExists(Path))
+    {
+    	BuildZenDashboard(Path, Parameters, Callback);
+    	return;
+    }
+	
+	constexpr bool bLaunchDetached = true;
+	constexpr bool bLaunchHidden = false;
+	constexpr bool bLaunchReallyHidden = false;
+
+	uint32 ProcessID = 0;
+	const int32 PriorityModifier = 0;
+
+	const TCHAR* OptionalWorkingDirectory = nullptr;
+
+	void* PipeWriteChild = nullptr;
+	void* PipeReadChild = nullptr;
+	ZenDashboardHandle = FPlatformProcess::CreateProc(*Path, *Parameters, bLaunchDetached, bLaunchHidden, bLaunchReallyHidden, &ProcessID, PriorityModifier, OptionalWorkingDirectory, PipeWriteChild, PipeReadChild);
+
+	if (ZenDashboardHandle.IsValid())
+	{
+		UE_LOG(LogZenDashboardLauncher, Log, TEXT("Launched Zen Dashboard executable: %s %s"), *Path, *Parameters);
+		Callback(EStartZenDashboardResult::Completed);
+	}
+	else
+	{
+		Callback(EStartZenDashboardResult::LaunchFailed);
+	}
+}
+
+void FZenDashboardLauncher::CloseZenDashboard()
+{
+	if (ZenDashboardHandle.IsValid())
+	{
+		FPlatformProcess::TerminateProc(ZenDashboardHandle);
+	}
+	else
+	{
+		UE_LOG(LogZenDashboardLauncher, Log, TEXT("Could not find the Zen Dashboard process handler"));
+	}
+}
+
+void FZenDashboardLauncher::BuildZenDashboard(const FString& Path, const FString& LaunchParameters, StartZenDashboardCallback Callback)
+{
+	UE_LOG(LogZenDashboardLauncher, Log, TEXT("Could not find the Zen Dashboard executable: %s. Attempting to build Zen Dashboard."), *Path);
+
+	FString Arguments;
+#if PLATFORM_WINDOWS
+	FText PlatformName = LOCTEXT("PlatformName_Windows", "Windows");
+	Arguments = TEXT("BuildTarget -Target=ZenDashboard -Platform=Win64");
+#elif PLATFORM_MAC
+	FText PlatformName = LOCTEXT("PlatformName_Mac", "Mac");
+	Arguments = TEXT("BuildTarget -Target=ZenDashboard -Platform=Mac");
+#elif PLATFORM_LINUX
+	FText PlatformName = LOCTEXT("PlatformName_Linux", "Linux");
+	Arguments = TEXT("BuildTarget -Target=ZenDashboard -Platform=Linux");
+#endif
+
+	IUATHelperModule::Get().CreateUatTask(Arguments, PlatformName, LOCTEXT("BuildingZenDashboard", "Building Zen Dashboard"),
+		LOCTEXT("BuildZenDashboardTask", "Build Zen Dashboard Task"), FAppStyle::GetBrush(TEXT("MainFrame.CookContent")), nullptr,
+		[this, Path, LaunchParameters, Callback](FString Result, double Time)
+		{
+			if (Result.Equals(TEXT("Completed")))
+			{
+#if PLATFORM_MAC
+				// On Mac we genereate the path again so that it includes the newly built executable.
+				FString NewPath = GetZenDashboardApplicationPath();
+				this->StartZenDashboard(NewPath, LaunchParameters, Callback);
+#else
+				this->StartZenDashboard(Path, LaunchParameters, Callback);
+#endif
+			}
+			else
+			{
+				Callback(EStartZenDashboardResult::BuildFailed);
+			}
+		});
+}
+
+} // namespace UE::Zen
+
+#undef LOCTEXT_NAMESPACE
+
+#endif // UE_WITH_ZEN
\ No newline at end of file
diff --git a/Engine/Source/Editor/DerivedDataEditor/Private/ZenDashboardLauncher.h b/Engine/Source/Editor/DerivedDataEditor/Private/ZenDashboardLauncher.h
new file mode 100644
--- /dev/null
+++ b/Engine/Source/Editor/DerivedDataEditor/Private/ZenDashboardLauncher.h
@@ -0,0 +1,85 @@
+// Copyright Epic Games, Inc. All Rights Reserved.
+
+#pragma once
+
+#include "CoreMinimal.h"
+#include "Experimental/ZenGlobals.h"
+#include "HAL/PlatformProcess.h"
+
+#if UE_WITH_ZEN
+
+#define UE_API
+
+namespace UE::Zen
+{
+
+class FZenDashboardLauncher : public TSharedFromThis<FZenDashboardLauncher>
+{
+public:
+	enum class EStartZenDashboardResult : uint32
+	{
+		Completed = 0,
+		BuildFailed = 1,
+		LaunchFailed = 2,
+	};
+
+	typedef TFunction<void(const EStartZenDashboardResult /*Result*/)> StartZenDashboardCallback;
+
+	UE_API FZenDashboardLauncher();
+	UE_API ~FZenDashboardLauncher();
+	
+	static UE_API const TSharedPtr<FZenDashboardLauncher>& Get()
+	{
+		if (!Instance)
+		{
+			Instance = MakeShared<FZenDashboardLauncher>();
+		}
+		return Instance;
+	}
+
+	/**
+	 * Returns the full path to ZenDashboard.exe for the current engine installation 
+	 */
+	UE_API FString GetZenDashboardApplicationPath();
+
+	/**
+	 * Launches the ZenDashboard executable from the given Path, displays an editor message if it fails.
+	 * If the executable is not found, a build process is started.
+	 * @param Path The full filename of ZenDashboard.exe to launch
+	 * @param Parameters The command line parameters to use when launching the exe
+	 */	
+	UE_API void StartZenDashboard(const FString& Path, const FString& Parameters = TEXT(""));
+
+	/**
+	* Launches the ZenDashboard executable from the given Path, displays an editor message if it fails.
+	* If the executable is not found, a build process is started.
+	* @param Path The full filename of ZenDashboard.exe to launch
+	* @param Parameters The command line parameters to use when launching the exe
+	* @param Callback A Callback that will be called when the launch process is completed.
+	*/
+	UE_API void StartZenDashboard(const FString& Path, const FString& Parameters, StartZenDashboardCallback Callback);
+
+	/**
+	* Closes ZenDashboard.exe.
+	*/
+	UE_API void CloseZenDashboard();
+
+private:
+	/*
+	 * Attempts building ZenDashboard via UAT, will launch with forwarded parameters if successful.
+	 * Assumes that the ZenDashboard Executable in path belongs to this engine.
+	 */
+	void BuildZenDashboard(const FString& Path, const FString& LaunchParameters, StartZenDashboardCallback Callback);
+
+private:
+	/** The process handler of ZenDashboard. */
+	FProcHandle ZenDashboardHandle;
+
+	static TSharedPtr<FZenDashboardLauncher> Instance;
+};
+
+} // namespace UE::Zen
+
+#undef UE_API
+
+#endif // UE_WITH_ZEN
diff --git a/Engine/Source/Editor/EditorStyle/Private/StarshipStyle.cpp b/Engine/Source/Editor/EditorStyle/Private/StarshipStyle.cpp
--- a/Engine/Source/Editor/EditorStyle/Private/StarshipStyle.cpp
+++ b/Engine/Source/Editor/EditorStyle/Private/StarshipStyle.cpp
@@ -1088,6 +1088,7 @@ void FStarshipEditorStyle::FStyle::SetupGeneralStyles()
 		Set("DerivedData.Cache.Settings", new IMAGE_BRUSH_SVG("Starship/DerivedData/DD_Cache_Settings", Icon16x16));
 		Set("DerivedData.Cache.Statistics", new IMAGE_BRUSH_SVG("Starship/DerivedData/DD_Cache_Statistics", Icon16x16));
 		Set("DerivedData.ResourceUsage", new IMAGE_BRUSH_SVG("Starship/DerivedData/DD_Resource_Usage", Icon16x16));
+		Set("DerivedData.LaunchZenDashboard", new CORE_IMAGE_BRUSH_SVG("Starship/Common/SessionFrontend", Icon16x16));
 		Set("DerivedData.RemoteCache.Uploading", new IMAGE_BRUSH_SVG("Starship/DerivedData/DD_Upload", Icon16x16, EStyleColor::AccentGreen));
 		Set("DerivedData.RemoteCache.Downloading", new IMAGE_BRUSH_SVG("Starship/DerivedData/DD_Download", Icon16x16, EStyleColor::AccentBlue));
 
diff --git a/Engine/Source/Programs/ZenDashboard/Private/ZenDashboard.cpp b/Engine/Source/Programs/ZenDashboard/Private/ZenDashboard.cpp
--- a/Engine/Source/Programs/ZenDashboard/Private/ZenDashboard.cpp
+++ b/Engine/Source/Programs/ZenDashboard/Private/ZenDashboard.cpp
@@ -60,7 +60,10 @@ static void HideOnCloseOverride(const TSharedRef<SWindow>& WindowBeingClosed)
 	WindowBeingClosed->HideWindow();
 }
 
-#if PLATFORM_WINDOWS
+// Controls whether we want the ZenDashboard app to behave as a systray app.  Currently only supported for windows.
+#define UE_ZENDASHBOARD_SYSTRAY 0
+
+#if UE_ZENDASHBOARD_SYSTRAY && PLATFORM_WINDOWS
 #define WM_TRAYICON (WM_USER + 1000)
 
 #define CMD_USER 400
@@ -72,7 +75,7 @@ static void HideOnCloseOverride(const TSharedRef<SWindow>& WindowBeingClosed)
 #endif
 
 class FZenDashboardApp
-#if PLATFORM_WINDOWS
+#if UE_ZENDASHBOARD_SYSTRAY && PLATFORM_WINDOWS
 	: IWindowsMessageHandler
 #endif
 {
@@ -85,7 +88,7 @@ class FZenDashboardApp
 
 	std::atomic<bool> bLatentExclusiveOperationActive = false;
 
-#if PLATFORM_WINDOWS
+#if UE_ZENDASHBOARD_SYSTRAY && PLATFORM_WINDOWS
 	bool ProcessMessage(HWND hwnd, uint32 msg, WPARAM wParam, LPARAM lParam, int32& OutResult) override
 	{
 		switch (msg)
@@ -165,7 +168,7 @@ class FZenDashboardApp
 
 		return false;
 	}
-#endif // PLATFORM_WINDOWS
+#endif // UE_ZENDASHBOARD_SYSTRAY && PLATFORM_WINDOWS
 
 	void ExitDashboard()
 	{
@@ -287,13 +290,13 @@ class FZenDashboardApp
 			MenuBuilder.AddPullDownMenu(
 				LOCTEXT( "FileMenu", "File" ),
 				LOCTEXT( "FileMenu_ToolTip", "Opens the file menu" ),
-				FNewMenuDelegate::CreateRaw( this, &FZenDashboardApp::FillFileMenu ) );
+				FOnGetContent::CreateRaw( this, &FZenDashboardApp::FillFileMenu ) );
 
 			// Control
 			MenuBuilder.AddPullDownMenu(
 				LOCTEXT( "ToolsMenu", "Tools" ),
 				LOCTEXT( "ToolsMenu_ToolTip", "Opens the tools menu" ),
-				FNewMenuDelegate::CreateRaw( this, &FZenDashboardApp::FillToolsMenu ) );
+				FOnGetContent::CreateRaw( this, &FZenDashboardApp::FillToolsMenu ) );
 		}
 
 		// Create the menu bar
@@ -303,8 +306,21 @@ class FZenDashboardApp
 		return MenuBarWidget;
 	}
 
-	void FillFileMenu(FMenuBuilder& MenuBuilder)
+	TSharedRef<SWidget> FillFileMenu()
 	{
+		const bool bCloseSelfOnly = false;
+		const bool bSearchable = false;
+		const bool bRecursivelySearchable = false;
+
+		FMenuBuilder MenuBuilder(true,
+			nullptr,
+			TSharedPtr<FExtender>(),
+			bCloseSelfOnly,
+			&FCoreStyle::Get(),
+			bSearchable,
+			NAME_None,
+			bRecursivelySearchable);
+
 		MenuBuilder.AddMenuEntry(
 			LOCTEXT("Exit", "Exit"),
 			LOCTEXT("Exit_ToolTip", "Exits the Zen Dashboard"),
@@ -316,10 +332,25 @@ class FZenDashboardApp
 			NAME_None,
 			EUserInterfaceActionType::Button
 		);
+
+		return MenuBuilder.MakeWidget();
 	}
 
-	void FillToolsMenu(FMenuBuilder& MenuBuilder)
+	TSharedRef<SWidget> FillToolsMenu()
 	{
+		const bool bCloseSelfOnly = false;
+		const bool bSearchable = false;
+		const bool bRecursivelySearchable = false;
+
+		FMenuBuilder MenuBuilder(true,
+			nullptr,
+			TSharedPtr<FExtender>(),
+			bCloseSelfOnly,
+			&FCoreStyle::Get(),
+			bSearchable,
+			NAME_None,
+			bRecursivelySearchable);
+
 		MenuBuilder.BeginSection(NAME_None, LOCTEXT("ServiceControls", "Service controls"));
 		{
 			MenuBuilder.AddMenuEntry(
@@ -356,9 +387,22 @@ class FZenDashboardApp
 				EUserInterfaceActionType::Button
 			);
 
-			MenuBuilder.AddSubMenu(LOCTEXT("Advanced", "Advanced"), LOCTEXT("Advanced_ToolTip", "Advanced service control commands"),
-				FNewMenuDelegate::CreateLambda([this](FMenuBuilder& SubMenuBuilder)
+			MenuBuilder.AddWrapperSubMenu(LOCTEXT("Advanced", "Advanced"), LOCTEXT("Advanced_ToolTip", "Advanced service control commands"),
+				FOnGetContent::CreateLambda([this]() -> TSharedRef<SWidget>
 					{
+						const bool bCloseSelfOnly = false;
+						const bool bSearchable = false;
+						const bool bRecursivelySearchable = false;
+
+						FMenuBuilder SubMenuBuilder(true,
+							nullptr,
+							TSharedPtr<FExtender>(),
+							bCloseSelfOnly,
+							&FCoreStyle::Get(),
+							bSearchable,
+							NAME_None,
+							bRecursivelySearchable);
+
 						SubMenuBuilder.AddMenuEntry(
 							LOCTEXT("DeleteDataAndRestart_ZenServer", "Delete data and restart Zen Server"),
 							LOCTEXT("DeleteDataAndRestart_ZenServer_ToolTip", "Stops ZenServer if it is running, deletes ALL of its data, then starts it from a blank state"),
@@ -370,7 +414,10 @@ class FZenDashboardApp
 							NAME_None,
 							EUserInterfaceActionType::Button
 						);
+
+						return SubMenuBuilder.MakeWidget();
 					})
+				, FSlateIcon()
 			);
 		}
 		MenuBuilder.EndSection();
@@ -442,6 +489,8 @@ class FZenDashboardApp
 			);
 		}
 		MenuBuilder.EndSection();
+
+		return MenuBuilder.MakeWidget();
 	}
 
 public:
@@ -459,8 +508,8 @@ public:
 
 	void Run()
 	{
-		const bool bSystemTrayMode = !!PLATFORM_WINDOWS;
-		const bool bShowWindow = !(bSystemTrayMode && FParse::Param(FCommandLine::Get(), TEXT("Minimized")));
+		const bool bSystemTrayMode = !!UE_ZENDASHBOARD_SYSTRAY && !!PLATFORM_WINDOWS;
+		const bool bShowWindow = !(FParse::Param(FCommandLine::Get(), TEXT("Minimized")) && bSystemTrayMode);
 		
 		Window =
 			SNew(SWindow)
@@ -545,18 +594,21 @@ public:
 		Slate.ClearKeyboardFocus(EFocusCause::Cleared);
 
 		// 
-#if PLATFORM_WINDOWS
-		((FWindowsApplication*)Slate.Get().GetPlatformApplication().Get())->AddMessageHandler(*this);
-
-		NOTIFYICONDATAW NotifyIconData;
-		memset(&NotifyIconData, 0, sizeof(NotifyIconData));
-		NotifyIconData.cbSize = sizeof(NotifyIconData);
-		NotifyIconData.hWnd = (HWND)Window->GetNativeWindow()->GetOSWindowHandle();
-		NotifyIconData.uID = 0;
-		NotifyIconData.uCallbackMessage = WM_USER + 1000;
-		NotifyIconData.uFlags = NIF_ICON | NIF_MESSAGE;
-		NotifyIconData.hIcon = ::LoadIconW(GetModuleHandleW(nullptr), MAKEINTRESOURCE(IDICON_UEGame));
-		Shell_NotifyIcon(NIM_ADD, &NotifyIconData);
+#if UE_ZENDASHBOARD_SYSTRAY && PLATFORM_WINDOWS
+		if (bSystemTrayMode)
+		{
+			((FWindowsApplication*)Slate.Get().GetPlatformApplication().Get())->AddMessageHandler(*this);
+
+			NOTIFYICONDATAW NotifyIconData;
+			memset(&NotifyIconData, 0, sizeof(NotifyIconData));
+			NotifyIconData.cbSize = sizeof(NotifyIconData);
+			NotifyIconData.hWnd = (HWND)Window->GetNativeWindow()->GetOSWindowHandle();
+			NotifyIconData.uID = 0;
+			NotifyIconData.uCallbackMessage = WM_USER + 1000;
+			NotifyIconData.uFlags = NIF_ICON | NIF_MESSAGE;
+			NotifyIconData.hIcon = ::LoadIconW(GetModuleHandleW(nullptr), MAKEINTRESOURCE(IDICON_UEGame));
+			Shell_NotifyIcon(NIM_ADD, &NotifyIconData);
+		}
 #endif
 
 		// loop until the app is ready to quit
@@ -578,12 +630,15 @@ public:
 			MainThreadTasks.Empty();
 		}
 
-#if PLATFORM_WINDOWS
-		memset(&NotifyIconData, 0, sizeof(NotifyIconData));
-		NotifyIconData.cbSize = sizeof(NotifyIconData);
-		NotifyIconData.uID = 0;
-		NotifyIconData.hWnd = (HWND)Window->GetNativeWindow()->GetOSWindowHandle();
-		Shell_NotifyIcon(NIM_DELETE, &NotifyIconData);
+#if UE_ZENDASHBOARD_SYSTRAY && PLATFORM_WINDOWS
+		if (bSystemTrayMode)
+		{
+			memset(&NotifyIconData, 0, sizeof(NotifyIconData));
+			NotifyIconData.cbSize = sizeof(NotifyIconData);
+			NotifyIconData.uID = 0;
+			NotifyIconData.hWnd = (HWND)Window->GetNativeWindow()->GetOSWindowHandle();
+			Shell_NotifyIcon(NIM_DELETE, &NotifyIconData);
+		}
 #endif
 
 		// Make sure the window is hidden, because it might take a while for the background thread to finish.
