Missing Zen Dashboard change as part of refactor

diff --git a/Engine/Source/Programs/ZenDashboard/Private/ZenDashboard.cpp b/Engine/Source/Programs/ZenDashboard/Private/ZenDashboard.cpp
--- a/Engine/Source/Programs/ZenDashboard/Private/ZenDashboard.cpp
+++ b/Engine/Source/Programs/ZenDashboard/Private/ZenDashboard.cpp
@@ -15,7 +15,7 @@
 #include "Modules/ModuleInterface.h"
 #include "Modules/ModuleManager.h"
 #include "RequiredProgramMainCPPInclude.h"
-#include "ServiceInstanceManager.h"
+#include "ZenServiceInstanceManager.h"
 #include "SMessageDialog.h"
 #include "StandaloneRenderer.h"
 #include "SZenCacheStatistics.h"
