Fix missing thread tag in ZenDashboard.

diff --git a/Engine/Source/Programs/ZenDashboard/Private/ZenDashboard.cpp b/Engine/Source/Programs/ZenDashboard/Private/ZenDashboard.cpp
--- a/Engine/Source/Programs/ZenDashboard/Private/ZenDashboard.cpp
+++ b/Engine/Source/Programs/ZenDashboard/Private/ZenDashboard.cpp
@@ -753,6 +753,8 @@ int ZenDashboardMain(const TCHAR* CmdLine)
 		FModuleManager::Get().UnloadModulesAtShutdown();
 		FEngineLoop::AppExit();
 	};
+	
+	const FTaskTagScope Scope(ETaskTag::EGameThread);
 
 	// start up the main loop
 	GEngineLoop.PreInit(CmdLine);
