Fix mac ZenDashboard icon filename and remove unused icons

diff --git a/Engine/Source/Programs/ZenDashboard/Resources/Mac/Assets.xcassets/AppIcon.appiconset/Contents.json b/Engine/Source/Programs/ZenDashboard/Resources/Mac/Assets.xcassets/AppIcon.appiconset/Contents.json
--- a/Engine/Source/Programs/ZenDashboard/Resources/Mac/Assets.xcassets/AppIcon.appiconset/Contents.json
+++ b/Engine/Source/Programs/ZenDashboard/Resources/Mac/Assets.xcassets/AppIcon.appiconset/Contents.json
@@ -49,7 +49,7 @@
       "size" : "256x256"
 	},
 	{
-	  "filename" : "icon_512_512.png",
+	  "filename" : "icon_512x512.png",
 	  "idiom" : "mac",
 	  "scale" : "1x",
 	  "size" : "512x512"
