Initialize CURL and by extension load the SSL module during initialization, to ensure the SSL module is loaded on the game thread before making requests from other threads.

[RN] none
#Synced-CL 48985904

diff --git a/Engine/Source/Programs/ZenDashboard/Private/ZenDashboard.cpp b/Engine/Source/Programs/ZenDashboard/Private/ZenDashboard.cpp
--- a/Engine/Source/Programs/ZenDashboard/Private/ZenDashboard.cpp
+++ b/Engine/Source/Programs/ZenDashboard/Private/ZenDashboard.cpp
@@ -24,6 +24,7 @@
 #include "SZenServiceStatus.h"
 #include "SZenWorkspaces.h"
 #include "Tasks/Task.h"
+#include "ThirdParty/UECurl.h"
 #include "Templates/SharedPointer.h"
 #include "Widgets/Input/SButton.h"
 #include "Widgets/Input/SButton.h"
@@ -768,6 +769,8 @@ int ZenDashboardMain(const TCHAR* CmdLine)
 
 	// Initialize high DPI mode
 	FSlateApplication::InitHighDPI(true);
+	UE::Curl::FCurlInitializeData CurlInitializeData{};
+	UE::Curl::ConditionalInitialize(CurlInitializeData);
 
 	{
 		// Create the platform slate application (what FSlateApplication::Get() returns)
@@ -805,6 +808,7 @@ int ZenDashboardMain(const TCHAR* CmdLine)
 
 		// Close down the Slate application
 		FSlateApplication::Shutdown();
+		UE::Curl::ConditionalShutdown(CurlInitializeData);
 	}
 
 	return true;
diff --git a/Engine/Source/Programs/ZenDashboard/ZenDashboard.Build.cs b/Engine/Source/Programs/ZenDashboard/ZenDashboard.Build.cs
--- a/Engine/Source/Programs/ZenDashboard/ZenDashboard.Build.cs
+++ b/Engine/Source/Programs/ZenDashboard/ZenDashboard.Build.cs
@@ -7,6 +7,7 @@ public class ZenDashboard : ModuleRules
 	public ZenDashboard(ReadOnlyTargetRules Target) : base(Target)
 	{
 		bTreatAsEngineModule = true;
+		PrivateDefinitions.Add("WITH_CURL=1");
 
 		PrivateDependencyModuleNames.AddRange(
 			new string[]
@@ -19,7 +20,8 @@ public class ZenDashboard : ModuleRules
 				"SlateCore",
 				"StandaloneRenderer",
 				"StorageServerWidgets",
-				"Zen"
+				"Zen",
+				"HTTP"
 			});
 
 		if (Target.IsInPlatformGroup(UnrealPlatformGroup.Linux))
@@ -32,6 +34,16 @@ public class ZenDashboard : ModuleRules
 			);
 		}
 
+		if (Target.Platform.IsInGroup(UnrealPlatformGroup.Windows) && Target.WindowsPlatform.bUseXCurl)
+		{
+			PublicDependencyModuleNames.Add("XCurl");
+		}
+		else
+		{
+			AddEngineThirdPartyPrivateStaticDependencies(Target, "libcurl");
+		}
+
+		AddEngineThirdPartyPrivateStaticDependencies(Target, "OpenSSL");
 		PublicIncludePathModuleNames.Add("Launch");
 	}
 }
