Rework the way curl is initialized when using the Zen module and ZenDashboard program.  The previous approach did not guarantee that the SSL  module was loaded from the game thread because WITH_SSL was not defined in ZenDashboard.cpp when UE::Curl::ConditionalInitialize was called.  Because UE::Curl::ConditionalInitialize is inline and is conditional on macro definitions, it could skip steps depending on the context in which it was compiled.

#Synced-CL 49552467

diff --git a/Engine/Source/Developer/Zen/Private/ZenModule.cpp b/Engine/Source/Developer/Zen/Private/ZenModule.cpp
--- a/Engine/Source/Developer/Zen/Private/ZenModule.cpp
+++ b/Engine/Source/Developer/Zen/Private/ZenModule.cpp
@@ -1,6 +1,9 @@
 // Copyright Epic Games, Inc. All Rights Reserved.
 
 #include "Modules/ModuleManager.h"
+#include "ThirdParty/UECurl.h"
+
+DEFINE_LOG_CATEGORY_STATIC(LogZenModule, Log, All);
 
 class FZenModule
 	: public IModuleInterface
@@ -10,7 +13,17 @@ public:
 	virtual void StartupModule() override
 	{
 		FModuleManager::Get().LoadModuleChecked(TEXT("Sockets"));
+		CURLcode InitResult = (CURLcode)UE::Curl::ConditionalInitialize(InitializeData);
+		if (InitResult != CURLE_OK)
+		{
+			UE_LOG(LogZenModule, Warning, TEXT("FZenModule could not initialize libcurl (result=%d)."), (int32)InitResult);
+		}
+	}
+	virtual void ShutdownModule() override
+	{
+		UE::Curl::ConditionalShutdown(InitializeData);
 	}
+	UE::Curl::FCurlInitializeData InitializeData;
 };
 
-IMPLEMENT_MODULE(FDefaultModuleImpl, Zen)
+IMPLEMENT_MODULE(FZenModule, Zen)
diff --git a/Engine/Source/Developer/Zen/Private/ZenServerHttp.cpp b/Engine/Source/Developer/Zen/Private/ZenServerHttp.cpp
--- a/Engine/Source/Developer/Zen/Private/ZenServerHttp.cpp
+++ b/Engine/Source/Developer/Zen/Private/ZenServerHttp.cpp
@@ -22,7 +22,6 @@
 #if PLATFORM_MICROSOFT
 #	include "Microsoft/HideMicrosoftPlatformTypes.h"
 #endif
-#include "ThirdParty/UECurl.h"
 
 #include "Logging/LogMacros.h"
 #include "Compression/CompressedBuffer.h"
@@ -44,13 +43,6 @@ LLM_DEFINE_TAG(ZenDDC, NAME_None, TEXT("DDCBackend"));
 
 DEFINE_LOG_CATEGORY_STATIC(LogZenHttp, Log, All);
 
-namespace UE::Curl
-{
-
-static FCurlInitializeData InitializeData;
-
-}
-
 namespace UE::Zen {
 
 #define UE_ZENDDC_BACKEND_WAIT_INTERVAL			0.01f
@@ -75,11 +67,6 @@ namespace UE::Zen {
 		, Domain(InEndpoint.GetURL())
 		, Endpoint(InEndpoint)
 	{
-		CURLcode InitResult = (CURLcode)UE::Curl::ConditionalInitialize(UE::Curl::InitializeData);
-		if (InitResult != CURLE_OK)
-		{
-			UE_LOG(LogZenHttp, Warning, TEXT("FZenHttpRequest could not initialize libcurl (result=%d)."), (int32)InitResult);
-		}
 		Curl = curl_easy_init();
 		Reset();
 	}
diff --git a/Engine/Source/Programs/ZenDashboard/Private/ZenDashboard.cpp b/Engine/Source/Programs/ZenDashboard/Private/ZenDashboard.cpp
--- a/Engine/Source/Programs/ZenDashboard/Private/ZenDashboard.cpp
+++ b/Engine/Source/Programs/ZenDashboard/Private/ZenDashboard.cpp
@@ -24,7 +24,6 @@
 #include "SZenServiceStatus.h"
 #include "SZenWorkspaces.h"
 #include "Tasks/Task.h"
-#include "ThirdParty/UECurl.h"
 #include "Templates/SharedPointer.h"
 #include "Widgets/Input/SButton.h"
 #include "Widgets/Input/SButton.h"
@@ -587,24 +586,31 @@ public:
 			Shell_NotifyIcon(NIM_ADD, &NotifyIconData);
 		}
 #endif
+		double DeltaTime = 0.0;
+		double LastTime = FPlatformTime::Seconds();
+		const float IdealFrameTime = 1.0f / 60;
+		const float BackgroundFrameTime = 1.0f / 4;
 
 		// loop until the app is ready to quit
 		while (!IsEngineExitRequested())
 		{
 			BeginExitIfRequested();
 
+			FTaskGraphInterface::Get().ProcessThreadUntilIdle(ENamedThreads::GameThread);
+			FTSTicker::GetCoreTicker().Tick(DeltaTime);
 			Slate.PumpMessages();
 			Slate.Tick();
 
-			FPlatformProcess::Sleep(1.0f / 30.0f);
+			const float FrameTime = !FPlatformApplicationMisc::IsThisApplicationForeground() && (FPlatformTime::Seconds() - FSlateApplication::Get().GetLastUserInteractionTime()) > 5 ? BackgroundFrameTime : IdealFrameTime;
+			FPlatformProcess::Sleep(FMath::Max<float>(0.f, FrameTime - static_cast<float>(FPlatformTime::Seconds() - LastTime)));
 
-			// Execute all the main thread tasks
-			FScopeLock Lock(&CriticalSection);
-			for (FSimpleDelegate& MainThreadTask : MainThreadTasks)
-			{
-				MainThreadTask.Execute();
-			}
-			MainThreadTasks.Empty();
+			DeltaTime = FPlatformTime::Seconds() - LastTime;
+			LastTime = FPlatformTime::Seconds();
+
+			UE::Stats::FStats::AdvanceFrame(false);
+			FCoreDelegates::OnEndFrame.Broadcast();
+
+			GFrameCounter++;
 		}
 
 #if UE_ZENDASHBOARD_SYSTRAY && PLATFORM_WINDOWS
@@ -769,8 +775,6 @@ int ZenDashboardMain(const TCHAR* CmdLine)
 
 	// Initialize high DPI mode
 	FSlateApplication::InitHighDPI(true);
-	UE::Curl::FCurlInitializeData CurlInitializeData{};
-	UE::Curl::ConditionalInitialize(CurlInitializeData);
 
 	{
 		// Create the platform slate application (what FSlateApplication::Get() returns)
@@ -808,7 +812,6 @@ int ZenDashboardMain(const TCHAR* CmdLine)
 
 		// Close down the Slate application
 		FSlateApplication::Shutdown();
-		UE::Curl::ConditionalShutdown(CurlInitializeData);
 	}
 
 	return true;
diff --git a/Engine/Source/Programs/ZenDashboard/ZenDashboard.Build.cs b/Engine/Source/Programs/ZenDashboard/ZenDashboard.Build.cs
--- a/Engine/Source/Programs/ZenDashboard/ZenDashboard.Build.cs
+++ b/Engine/Source/Programs/ZenDashboard/ZenDashboard.Build.cs
@@ -14,14 +14,15 @@ public class ZenDashboard : ModuleRules
 			{
 				"ApplicationCore",
 				"Core",
+				"HTTP",
 				"Json",
 				"Projects",
 				"Slate",
 				"SlateCore",
+				"SSL",
 				"StandaloneRenderer",
 				"StorageServerWidgets",
-				"Zen",
-				"HTTP"
+				"Zen"
 			});
 
 		if (Target.IsInPlatformGroup(UnrealPlatformGroup.Linux))
